Specify a short name for DMXPixelMappingBlueprintGraph intermediates to avoid path length warnings
#rnx

diff --git a/Engine/Plugins/Experimental/DigitalMultiplex/DMXPixelMapping/Source/DMXPixelMappingBlueprintGraph/DMXPixelMappingBlueprintGraph.Build.cs b/Engine/Plugins/Experimental/DigitalMultiplex/DMXPixelMapping/Source/DMXPixelMappingBlueprintGraph/DMXPixelMappingBlueprintGraph.Build.cs
--- a/Engine/Plugins/Experimental/DigitalMultiplex/DMXPixelMapping/Source/DMXPixelMappingBlueprintGraph/DMXPixelMappingBlueprintGraph.Build.cs
+++ b/Engine/Plugins/Experimental/DigitalMultiplex/DMXPixelMapping/Source/DMXPixelMappingBlueprintGraph/DMXPixelMappingBlueprintGraph.Build.cs
@@ -6,6 +6,8 @@ public class DMXPixelMappingBlueprintGraph : ModuleRules
 {
 	public DMXPixelMappingBlueprintGraph(ReadOnlyTargetRules Target) : base(Target)
 	{
+		ShortName = "DMXPxlMapBPGraph";
+		
 		PCHUsage = ModuleRules.PCHUsageMode.UseExplicitOrSharedPCHs;
 		
 		PublicDependencyModuleNames.AddRange(
