Quick pass at some plugin category clean up
- Moved various things from "Programming" and "Blueprints" to more domain-specific categories
- Merged "Media" and "Movie Players" into "Media Players"
- Merged "Automation" into "Testing"
- Merged "Widgets" and "Web" into "UI"
etc...

diff --git a/Engine/Plugins/Developer/AnimationSharing/AnimationSharing.uplugin b/Engine/Plugins/Developer/AnimationSharing/AnimationSharing.uplugin
--- a/Engine/Plugins/Developer/AnimationSharing/AnimationSharing.uplugin
+++ b/Engine/Plugins/Developer/AnimationSharing/AnimationSharing.uplugin
@@ -1,10 +1,10 @@
 {
-	"FileVersion": 1,
+	"FileVersion": 3,
 	"Version": 1,
 	"VersionName": "1.0",
 	"FriendlyName": "Animation Sharing",
 	"Description": "Plugin to create Shared Animation systems using the Master-Child pose functionality",
-	"Category": "Programming",
+	"Category": "Animation",
 	"CreatedBy": "Epic Games, Inc.",
 	"CreatedByURL": "http://epicgames.com",
 	"DocsURL": "",
@@ -16,9 +16,9 @@
 	"Installed": false,
 	"Modules": [
 		{
-		  "Name": "AnimationSharing",
-		  "Type": "Runtime",
-		  "LoadingPhase": "PreDefault"
+			"Name": "AnimationSharing",
+			"Type": "Runtime",
+			"LoadingPhase": "PreDefault"
 		},
 		{
 			"Name": "AnimationSharingEd",
@@ -26,10 +26,11 @@
 			"LoadingPhase": "Default"
 		}
 	],
-	"Plugins": [ 
-		{ 
-			"Name": "SignificanceManager", 
-			"Enabled": true 
-		} 
-	]
+	"Plugins": [
+		{
+			"Name": "SignificanceManager",
+			"Enabled": true
+		}
+	],
+	"IsExperimentalVersion": false
 }
\ No newline at end of file
diff --git a/Engine/Plugins/Developer/Concert/MultiUserClient/MultiUserClient.uplugin b/Engine/Plugins/Developer/Concert/MultiUserClient/MultiUserClient.uplugin
--- a/Engine/Plugins/Developer/Concert/MultiUserClient/MultiUserClient.uplugin
+++ b/Engine/Plugins/Developer/Concert/MultiUserClient/MultiUserClient.uplugin
@@ -4,7 +4,7 @@
 	"VersionName": "0.2",
 	"FriendlyName": "Multi-User Editing",
 	"Description": "Allow collaborative multi-users sessions in the Editor",
-	"Category": "Developer",
+	"Category": "Editor",
 	"CreatedBy": "Epic Games, Inc.",
 	"CreatedByURL": "http://epicgames.com",
 	"DocsURL": "",
@@ -35,5 +35,6 @@
 			"Name": "MultiUserTakes",
 			"Enabled": true
 		}
-	]
-}
+	],
+	"IsExperimentalVersion": false
+}
\ No newline at end of file
diff --git a/Engine/Plugins/Developer/PropertyAccessNode/PropertyAccessNode.uplugin b/Engine/Plugins/Developer/PropertyAccessNode/PropertyAccessNode.uplugin
--- a/Engine/Plugins/Developer/PropertyAccessNode/PropertyAccessNode.uplugin
+++ b/Engine/Plugins/Developer/PropertyAccessNode/PropertyAccessNode.uplugin
@@ -4,7 +4,7 @@
 	"VersionName": "1.0",
 	"FriendlyName": "Property Access Node",
 	"Description": "Blueprint node that allows access to properties via a property path",
-	"Category": "Developer",
+	"Category": "Blueprints",
 	"CreatedBy": "Epic Games, Inc.",
 	"CreatedByURL": "http://epicgames.com",
 	"DocsURL": "",
@@ -14,20 +14,18 @@
 	"CanContainContent": false,
 	"IsBetaVersion": false,
 	"Installed": false,
-
-	"Modules":
-	[
+	"Modules": [
 		{
 			"Name": "PropertyAccessNode",
 			"Type": "UncookedOnly",
 			"LoadingPhase": "PreDefault"
 		}
 	],
-	"Plugins":
-	[
+	"Plugins": [
 		{
 			"Name": "PropertyAccessEditor",
 			"Enabled": true
 		}
-	]
+	],
+	"IsExperimentalVersion": false
 }
\ No newline at end of file
diff --git a/Engine/Plugins/Editor/InterchangeEditor/InterchangeEditor.uplugin b/Engine/Plugins/Editor/InterchangeEditor/InterchangeEditor.uplugin
--- a/Engine/Plugins/Editor/InterchangeEditor/InterchangeEditor.uplugin
+++ b/Engine/Plugins/Editor/InterchangeEditor/InterchangeEditor.uplugin
@@ -4,7 +4,7 @@
 	"VersionName": "1.0",
 	"FriendlyName": "Interchange Editor",
 	"Description": "Interchange Editor plugin pull on Interchange plugin to allow us to force interchange plugin in editor.",
-	"Category": "Import Export",
+	"Category": "Importers",
 	"CreatedBy": "Epic Games, Inc.",
 	"CreatedByURL": "http://epicgames.com",
 	"DocsURL": "",
@@ -30,5 +30,6 @@
 				"Program"
 			]
 		}
-	]
+	],
+	"IsExperimentalVersion": false
 }
\ No newline at end of file
diff --git a/Engine/Plugins/Experimental/AutomationUtils/AutomationUtils.uplugin b/Engine/Plugins/Experimental/AutomationUtils/AutomationUtils.uplugin
--- a/Engine/Plugins/Experimental/AutomationUtils/AutomationUtils.uplugin
+++ b/Engine/Plugins/Experimental/AutomationUtils/AutomationUtils.uplugin
@@ -1,10 +1,10 @@
 {
-	"FileVersion": 1,
+	"FileVersion": 3,
 	"Version": 1,
 	"VersionName": "0.1",
 	"FriendlyName": "Automation Utilities",
 	"Description": "Tools and Utilities for Automation purposes",
-	"Category": "Automation",
+	"Category": "Testing",
 	"CreatedBy": "Epic Games, Inc.",
 	"CreatedByURL": "http://epicgames.com",
 	"MarketplaceURL": "",
@@ -25,11 +25,12 @@
 			"LoadingPhase": "Default"
 		}
 	],
-	"Plugins":
-	[
+	"Plugins": [
 		{
 			"Name": "ScreenshotTools",
 			"Enabled": true
 		}
-	]
+	],
+	"DocsURL": "",
+	"IsExperimentalVersion": false
 }
\ No newline at end of file
diff --git a/Engine/Plugins/Experimental/DigitalMultiplex/DMXPixelMapping/DMXPixelMapping.uplugin b/Engine/Plugins/Experimental/DigitalMultiplex/DMXPixelMapping/DMXPixelMapping.uplugin
--- a/Engine/Plugins/Experimental/DigitalMultiplex/DMXPixelMapping/DMXPixelMapping.uplugin
+++ b/Engine/Plugins/Experimental/DigitalMultiplex/DMXPixelMapping/DMXPixelMapping.uplugin
@@ -1,57 +1,59 @@
 {
-  "FileVersion": 1,
-  "Version": 1,
-  "VersionName": "1.0",
-  "FriendlyName": "DMX Pixel Mapping",
-  "Description": "Tools set for map LED digital pixel strip or fixture arrays regardless of shape or size",
-  "Category": "Media",
-  "CreatedBy": "Epic Games, Inc.",
-  "CreatedByURL": "http://epicgames.com",
-  "DocsURL": "",
-  "MarketplaceURL": "",
-  "SupportURL": "",
-  "CanContainContent": true,
-  "IsBetaVersion": true,
-  "Modules": [
-    {
-      "Name": "DMXPixelMappingCore",
-      "Type": "Runtime",
-      "LoadingPhase": "Default"
-    },
-    {
-      "Name": "DMXPixelMappingRuntime",
-      "Type": "Runtime",
-      "LoadingPhase": "Default"
-    },
-    {
-      "Name": "DMXPixelMappingEditor",
-      "Type": "Editor",
-      "LoadingPhase": "Default"
-    },
-      {
-      "Name": "DMXPixelMappingEditorWidgets",
-      "Type": "Editor",
-      "LoadingPhase": "Default"
-    },
-    {
-      "Name": "DMXPixelMappingBlueprintGraph",
-      "Type": "UncookedOnly",
-      "LoadingPhase": "Default"
-    },
-    {
-      "Name": "DMXPixelMappingRenderer",
-      "Type": "Runtime",
-      "LoadingPhase": "PostConfigInit"
-    }
-  ],
-  "Plugins": [
-    {
-      "Name": "DMXProtocol",
-      "Enabled": true
-    },
-	{
-      "Name": "DMXEngine",
-      "Enabled": true
-    }
-  ]
-}
+	"FileVersion": 3,
+	"Version": 1,
+	"VersionName": "1.0",
+	"FriendlyName": "DMX Pixel Mapping",
+	"Description": "Tools set for map LED digital pixel strip or fixture arrays regardless of shape or size",
+	"Category": "Virtual Production",
+	"CreatedBy": "Epic Games, Inc.",
+	"CreatedByURL": "http://epicgames.com",
+	"DocsURL": "",
+	"MarketplaceURL": "",
+	"SupportURL": "",
+	"CanContainContent": true,
+	"IsBetaVersion": true,
+	"Modules": [
+		{
+			"Name": "DMXPixelMappingCore",
+			"Type": "Runtime",
+			"LoadingPhase": "Default"
+		},
+		{
+			"Name": "DMXPixelMappingRuntime",
+			"Type": "Runtime",
+			"LoadingPhase": "Default"
+		},
+		{
+			"Name": "DMXPixelMappingEditor",
+			"Type": "Editor",
+			"LoadingPhase": "Default"
+		},
+		{
+			"Name": "DMXPixelMappingEditorWidgets",
+			"Type": "Editor",
+			"LoadingPhase": "Default"
+		},
+		{
+			"Name": "DMXPixelMappingBlueprintGraph",
+			"Type": "UncookedOnly",
+			"LoadingPhase": "Default"
+		},
+		{
+			"Name": "DMXPixelMappingRenderer",
+			"Type": "Runtime",
+			"LoadingPhase": "PostConfigInit"
+		}
+	],
+	"Plugins": [
+		{
+			"Name": "DMXProtocol",
+			"Enabled": true
+		},
+		{
+			"Name": "DMXEngine",
+			"Enabled": true
+		}
+	],
+	"IsExperimentalVersion": false,
+	"Installed": false
+}
\ No newline at end of file
diff --git a/Engine/Plugins/Experimental/Gauntlet/Gauntlet.uplugin b/Engine/Plugins/Experimental/Gauntlet/Gauntlet.uplugin
--- a/Engine/Plugins/Experimental/Gauntlet/Gauntlet.uplugin
+++ b/Engine/Plugins/Experimental/Gauntlet/Gauntlet.uplugin
@@ -1,25 +1,25 @@
 {
-	"FileVersion" : 3,
-	"Version" : 1,
-	"VersionName" : "0.1",
-	"FriendlyName" : "Gauntlet",
-	"Description" : "Provides a helper class for creating and managing tests in your game",
-	"Category" : "Automation",
-	"CreatedBy" : "Epic Games",
-	"CreatedByURL" : "",
-	"DocsURL" : "",
-	"MarketplaceURL" : "",
-	"SupportURL" : "",
-	"EnabledByDefault" : false,
-	"CanContainContent" : false,
-	"IsBetaVersion" : true,
-	"Installed" : false,
-	"Modules" :
-	[
+	"FileVersion": 3,
+	"Version": 1,
+	"VersionName": "0.1",
+	"FriendlyName": "Gauntlet",
+	"Description": "Provides a helper class for creating and managing tests in your game",
+	"Category": "Testing",
+	"CreatedBy": "Epic Games",
+	"CreatedByURL": "",
+	"DocsURL": "",
+	"MarketplaceURL": "",
+	"SupportURL": "",
+	"EnabledByDefault": false,
+	"CanContainContent": false,
+	"IsBetaVersion": true,
+	"Installed": false,
+	"Modules": [
 		{
-			"Name" : "Gauntlet",
-			"Type" : "Runtime",
+			"Name": "Gauntlet",
+			"Type": "Runtime",
 			"LoadingPhase": "PreDefault"
 		}
-	]
+	],
+	"IsExperimentalVersion": false
 }
\ No newline at end of file
diff --git a/Engine/Plugins/Experimental/GizmoFramework/GizmoFramework.uplugin b/Engine/Plugins/Experimental/GizmoFramework/GizmoFramework.uplugin
--- a/Engine/Plugins/Experimental/GizmoFramework/GizmoFramework.uplugin
+++ b/Engine/Plugins/Experimental/GizmoFramework/GizmoFramework.uplugin
@@ -4,7 +4,7 @@
 	"VersionName": "1.0",
 	"FriendlyName": "New TRS Gizmos",
 	"Description": "Toggle use of InteractiveToolsFramework based TRS Gizmos",
-	"Category": "Assets",
+	"Category": "Experimental",
 	"CreatedBy": "Epic Games, Inc.",
 	"CreatedByURL": "http://epicgames.com",
 	"DocsURL": "",
@@ -21,10 +21,10 @@
 			"LoadingPhase": "Default"
 		}
 	],
-	"Plugins" : [
+	"Plugins": [
 		{
-			"Name" : "GizmoEdMode",
-			"Enabled" : true
+			"Name": "GizmoEdMode",
+			"Enabled": true
 		}
 	]
-}
+}
\ No newline at end of file
diff --git a/Engine/Plugins/Media/MediaFrameworkUtilities/MediaFrameworkUtilities.uplugin b/Engine/Plugins/Media/MediaFrameworkUtilities/MediaFrameworkUtilities.uplugin
--- a/Engine/Plugins/Media/MediaFrameworkUtilities/MediaFrameworkUtilities.uplugin
+++ b/Engine/Plugins/Media/MediaFrameworkUtilities/MediaFrameworkUtilities.uplugin
@@ -1,24 +1,23 @@
 {
-	"FileVersion" : 3,
-	"Version" : 1,
-	"VersionName" : "1.0",
-	"FriendlyName" : "Media Framework Utilities",
-	"Description" : "Utility assets and actors to ease the use of the Media Framework.",
-	"Category" : "Media",
-	"CreatedBy" : "Epic Games, Inc.",
-	"CreatedByURL" : "http://epicgames.com",
-	"MarketplaceURL" : "",
-	"SupportURL" : "",
-	"EnabledByDefault" : false,
-	"CanContainContent" : true,
-	"IsBetaVersion" : false,
-	"Installed" : false,
-	"Modules" :
-	[
+	"FileVersion": 3,
+	"Version": 1,
+	"VersionName": "1.0",
+	"FriendlyName": "Media Framework Utilities",
+	"Description": "Utility assets and actors to ease the use of the Media Framework.",
+	"Category": "Media Players",
+	"CreatedBy": "Epic Games, Inc.",
+	"CreatedByURL": "http://epicgames.com",
+	"MarketplaceURL": "",
+	"SupportURL": "",
+	"EnabledByDefault": false,
+	"CanContainContent": true,
+	"IsBetaVersion": false,
+	"Installed": false,
+	"Modules": [
 		{
-			"Name" : "MediaFrameworkUtilities",
-			"Type" : "Runtime",
-			"LoadingPhase" : "Default",
+			"Name": "MediaFrameworkUtilities",
+			"Type": "Runtime",
+			"LoadingPhase": "Default",
 			"WhitelistPlatforms": [
 				"Win64",
 				"Linux",
@@ -26,9 +25,9 @@
 			]
 		},
 		{
-			"Name" : "MediaFrameworkUtilitiesEditor",
-			"Type" : "Editor",
-			"LoadingPhase" : "PostEngineInit",
+			"Name": "MediaFrameworkUtilitiesEditor",
+			"Type": "Editor",
+			"LoadingPhase": "PostEngineInit",
 			"WhitelistPlatforms": [
 				"Win64",
 				"Linux",
@@ -36,19 +35,20 @@
 			]
 		}
 	],
-	"Plugins" :
-	[
+	"Plugins": [
 		{
 			"Name": "MediaIOFramework",
 			"Enabled": true
 		},
 		{
-			"Name" : "OpenCVLensDistortion",
-			"Enabled" : true
+			"Name": "OpenCVLensDistortion",
+			"Enabled": true
 		},
 		{
 			"Name": "MediaPlayerEditor",
 			"Enabled": true
 		}
-	]
-}
+	],
+	"DocsURL": "",
+	"IsExperimentalVersion": false
+}
\ No newline at end of file
diff --git a/Engine/Plugins/Media/MediaIOFramework/MediaIOFramework.uplugin b/Engine/Plugins/Media/MediaIOFramework/MediaIOFramework.uplugin
--- a/Engine/Plugins/Media/MediaIOFramework/MediaIOFramework.uplugin
+++ b/Engine/Plugins/Media/MediaIOFramework/MediaIOFramework.uplugin
@@ -1,18 +1,18 @@
 {
-	"FileVersion" : 3,
-	"Version" : 1,
-	"VersionName" : "1.0",
-	"FriendlyName" : "Media IO Framework",
-	"Description" : "Media Framework classes to support Professional Media IO used by the Virtual Production industry.",
-	"Category" : "Media",
-	"CreatedBy" : "Epic Games, Inc.",
-	"CreatedByURL" : "http://epicgames.com",
-	"MarketplaceURL" : "",
-	"SupportURL" : "",
-	"EnabledByDefault" : false,
-	"CanContainContent" : false,
-	"IsBetaVersion" : false,
-	"Installed" : false,
+	"FileVersion": 3,
+	"Version": 1,
+	"VersionName": "1.0",
+	"FriendlyName": "Media IO Framework",
+	"Description": "Media Framework classes to support Professional Media IO used by the Virtual Production industry.",
+	"Category": "Virtual Production",
+	"CreatedBy": "Epic Games, Inc.",
+	"CreatedByURL": "http://epicgames.com",
+	"MarketplaceURL": "",
+	"SupportURL": "",
+	"EnabledByDefault": false,
+	"CanContainContent": false,
+	"IsBetaVersion": false,
+	"Installed": false,
 	"Modules": [
 		{
 			"Name": "MediaIOCore",
@@ -24,5 +24,7 @@
 			"Type": "Editor",
 			"LoadingPhase": "PostEngineInit"
 		}
-	]
-}
+	],
+	"DocsURL": "",
+	"IsExperimentalVersion": false
+}
\ No newline at end of file
diff --git a/Engine/Plugins/Media/TimecodeSynchronizer/TimecodeSynchronizer.uplugin b/Engine/Plugins/Media/TimecodeSynchronizer/TimecodeSynchronizer.uplugin
--- a/Engine/Plugins/Media/TimecodeSynchronizer/TimecodeSynchronizer.uplugin
+++ b/Engine/Plugins/Media/TimecodeSynchronizer/TimecodeSynchronizer.uplugin
@@ -1,19 +1,19 @@
 {
-	"FileVersion" : 3,
-	"Version" : 1,
-	"VersionName" : "0.1",
-	"FriendlyName" : "Timecode Synchronizer",
-	"Description" : "An asset that will become the TimecodeProvider once all the inputs get synchronized to a timecode.",
-	"Category" : "Media",
-	"CreatedBy" : "Epic Games, Inc.",
-	"CreatedByURL" : "http://epicgames.com",
-	"DocsURL" : "",
-	"MarketplaceURL" : "",
-	"SupportURL" : "",
-	"EnabledByDefault" : false,
-	"CanContainContent" : false,
-	"IsBetaVersion" : true,
-	"Installed" : false,
+	"FileVersion": 3,
+	"Version": 1,
+	"VersionName": "0.1",
+	"FriendlyName": "Timecode Synchronizer",
+	"Description": "An asset that will become the TimecodeProvider once all the inputs get synchronized to a timecode.",
+	"Category": "Virtual Production",
+	"CreatedBy": "Epic Games, Inc.",
+	"CreatedByURL": "http://epicgames.com",
+	"DocsURL": "",
+	"MarketplaceURL": "",
+	"SupportURL": "",
+	"EnabledByDefault": false,
+	"CanContainContent": false,
+	"IsBetaVersion": true,
+	"Installed": false,
 	"Modules": [
 		{
 			"Name": "TimecodeSynchronizer",
@@ -31,5 +31,6 @@
 			"Name": "MediaPlayerEditor",
 			"Enabled": true
 		}
-	]
+	],
+	"IsExperimentalVersion": false
 }
\ No newline at end of file
diff --git a/Engine/Plugins/MeshPainting/MeshPainting.uplugin b/Engine/Plugins/MeshPainting/MeshPainting.uplugin
--- a/Engine/Plugins/MeshPainting/MeshPainting.uplugin
+++ b/Engine/Plugins/MeshPainting/MeshPainting.uplugin
@@ -4,7 +4,7 @@
 	"VersionName": "1.0",
 	"FriendlyName": "Mesh Painting",
 	"Description": "System for painting data onto meshes.",
-	"Category": "Mesh",
+	"Category": "Editor",
 	"CreatedBy": "Epic Games, Inc.",
 	"CreatedByURL": "http://epicgames.com",
 	"DocsURL": "",
@@ -31,5 +31,6 @@
 			"Name": "GeometryProcessing",
 			"Enabled": true
 		}
-	]
-}
+	],
+	"IsExperimentalVersion": false
+}
\ No newline at end of file
diff --git a/Engine/Plugins/Portal/LauncherChunkInstaller/LauncherChunkInstaller.uplugin b/Engine/Plugins/Portal/LauncherChunkInstaller/LauncherChunkInstaller.uplugin
--- a/Engine/Plugins/Portal/LauncherChunkInstaller/LauncherChunkInstaller.uplugin
+++ b/Engine/Plugins/Portal/LauncherChunkInstaller/LauncherChunkInstaller.uplugin
@@ -1,27 +1,30 @@
 {
-	"FileVersion" : 3,
-	
-	"FriendlyName" : "Launcher Chunk Installer",
-	"Version" : 1,
-	"VersionName" : "1.0",
-	"Description" : "Chunk installer module that hooks into launcher",
-	"Category" : "Launcher",
-	"CreatedBy" : "Epic Games, Inc.",
-	"CreatedByURL" : "http://epicgames.com",
-	"EnabledByDefault" : true,
-		
-	 "Modules": 
-	 [
-        {
-            "Name": "LauncherChunkInstaller",
-            "Type": "Runtime",
-            "LoadingPhase" : "PostConfigInit",
-			"WhitelistPlatforms" :
-			[
+	"FileVersion": 3,
+	"FriendlyName": "Launcher Chunk Installer",
+	"Version": 1,
+	"VersionName": "1.0",
+	"Description": "Chunk installer module that hooks into launcher",
+	"Category": "Other",
+	"CreatedBy": "Epic Games, Inc.",
+	"CreatedByURL": "http://epicgames.com",
+	"EnabledByDefault": true,
+	"Modules": [
+		{
+			"Name": "LauncherChunkInstaller",
+			"Type": "Runtime",
+			"LoadingPhase": "PostConfigInit",
+			"WhitelistPlatforms": [
 				"Win64",
 				"Linux",
 				"Mac"
 			]
-        }
-    ]
+		}
+	],
+	"DocsURL": "",
+	"MarketplaceURL": "",
+	"SupportURL": "",
+	"CanContainContent": false,
+	"IsBetaVersion": false,
+	"IsExperimentalVersion": false,
+	"Installed": false
 }
\ No newline at end of file
diff --git a/Engine/Plugins/Runtime/ActorLayerUtilities/ActorLayerUtilities.uplugin b/Engine/Plugins/Runtime/ActorLayerUtilities/ActorLayerUtilities.uplugin
--- a/Engine/Plugins/Runtime/ActorLayerUtilities/ActorLayerUtilities.uplugin
+++ b/Engine/Plugins/Runtime/ActorLayerUtilities/ActorLayerUtilities.uplugin
@@ -1,31 +1,30 @@
 {
-	"FileVersion" : 3,
-	"Version" : 1,
-	"VersionName" : "1.0",
-	"FriendlyName" : "Actor Layer Utilities",
-	"Description" : "Utilites for interacting with actor layers from blueprints",
-	"Category" : "Developer",
-	"CreatedBy" : "",
-	"CreatedByURL" : "",
-	"DocsURL" : "",
-	"MarketplaceURL" : "",
-	"SupportURL" : "",
-	"EnabledByDefault" : true,
-	"CanContainContent" : false,
-	"IsBetaVersion" : false,
-	"Installed" : false,
-	"Modules" :
-	[
+	"FileVersion": 3,
+	"Version": 1,
+	"VersionName": "1.0",
+	"FriendlyName": "Actor Layer Utilities",
+	"Description": "Utilites for interacting with actor layers from blueprints",
+	"Category": "Editor",
+	"CreatedBy": "",
+	"CreatedByURL": "",
+	"DocsURL": "",
+	"MarketplaceURL": "",
+	"SupportURL": "",
+	"EnabledByDefault": true,
+	"CanContainContent": false,
+	"IsBetaVersion": false,
+	"Installed": false,
+	"Modules": [
 		{
-			"Name" : "ActorLayerUtilities",
-			"Type" : "Runtime",
-			"LoadingPhase" : "Default"
+			"Name": "ActorLayerUtilities",
+			"Type": "Runtime",
+			"LoadingPhase": "Default"
 		},
-
 		{
-			"Name" : "ActorLayerUtilitiesEditor",
-			"Type" : "Editor",
-			"LoadingPhase" : "Default"
+			"Name": "ActorLayerUtilitiesEditor",
+			"Type": "Editor",
+			"LoadingPhase": "Default"
 		}
-	]
+	],
+	"IsExperimentalVersion": false
 }
\ No newline at end of file
diff --git a/Engine/Plugins/Runtime/AndroidMoviePlayer/AndroidMoviePlayer.uplugin b/Engine/Plugins/Runtime/AndroidMoviePlayer/AndroidMoviePlayer.uplugin
--- a/Engine/Plugins/Runtime/AndroidMoviePlayer/AndroidMoviePlayer.uplugin
+++ b/Engine/Plugins/Runtime/AndroidMoviePlayer/AndroidMoviePlayer.uplugin
@@ -1,29 +1,28 @@
 {
-	"FileVersion" : 3,
-	"Version" : 1,
-	"VersionName" : "1.0",
-	"FriendlyName" : "Android Movie Player",
-	"Description" : "Android Platform Movie Player using Android Media library",
-	"Category" : "Movie Players",
-	"CreatedBy" : "Epic Games, Inc.",
-	"CreatedByURL" : "http://epicgames.com",
-	"DocsURL" : "",
-	"MarketplaceURL" : "",
-	"SupportURL" : "",
-	"EnabledByDefault" : true,
-	"CanContainContent" : false,
-	"IsBetaVersion" : false,
-	"Installed" : false,
-	"Modules" :
-	[
+	"FileVersion": 3,
+	"Version": 1,
+	"VersionName": "1.0",
+	"FriendlyName": "Android Movie Player",
+	"Description": "Android Platform Movie Player using Android Media library",
+	"Category": "Media Players",
+	"CreatedBy": "Epic Games, Inc.",
+	"CreatedByURL": "http://epicgames.com",
+	"DocsURL": "",
+	"MarketplaceURL": "",
+	"SupportURL": "",
+	"EnabledByDefault": true,
+	"CanContainContent": false,
+	"IsBetaVersion": false,
+	"Installed": false,
+	"Modules": [
 		{
-			"Name" : "AndroidMoviePlayer",
-			"Type" : "RuntimeNoCommandlet",
-			"LoadingPhase" : "PreLoadingScreen",
-			"WhitelistPlatforms" :
-			[
+			"Name": "AndroidMoviePlayer",
+			"Type": "RuntimeNoCommandlet",
+			"LoadingPhase": "PreLoadingScreen",
+			"WhitelistPlatforms": [
 				"Android"
 			]
 		}
-	]
+	],
+	"IsExperimentalVersion": false
 }
\ No newline at end of file
diff --git a/Engine/Plugins/Runtime/AnimationBudgetAllocator/AnimationBudgetAllocator.uplugin b/Engine/Plugins/Runtime/AnimationBudgetAllocator/AnimationBudgetAllocator.uplugin
--- a/Engine/Plugins/Runtime/AnimationBudgetAllocator/AnimationBudgetAllocator.uplugin
+++ b/Engine/Plugins/Runtime/AnimationBudgetAllocator/AnimationBudgetAllocator.uplugin
@@ -1,10 +1,10 @@
 {
-	"FileVersion": 1,
+	"FileVersion": 3,
 	"Version": 1,
 	"VersionName": "1.0",
 	"FriendlyName": "Animation Budget Allocator",
 	"Description": "Constrains the time taken for animation to run by dynamically throttling skeletal mesh component ticking.",
-	"Category": "Programming",
+	"Category": "Animation",
 	"CreatedBy": "Epic Games, Inc.",
 	"CreatedByURL": "http://epicgames.com",
 	"DocsURL": "",
@@ -20,5 +20,6 @@
 			"Type": "Runtime",
 			"LoadingPhase": "PreDefault"
 		}
-	]
+	],
+	"IsExperimentalVersion": false
 }
\ No newline at end of file
diff --git a/Engine/Plugins/Runtime/AppleMoviePlayer/AppleMoviePlayer.uplugin b/Engine/Plugins/Runtime/AppleMoviePlayer/AppleMoviePlayer.uplugin
--- a/Engine/Plugins/Runtime/AppleMoviePlayer/AppleMoviePlayer.uplugin
+++ b/Engine/Plugins/Runtime/AppleMoviePlayer/AppleMoviePlayer.uplugin
@@ -1,31 +1,30 @@
 {
-	"FileVersion" : 3,
-	"Version" : 1,
-	"VersionName" : "1.0",
-	"FriendlyName" : "Apple Movie Player",
-	"Description" : "Apple Platform Movie Player using AVPlayer library",
-	"Category" : "Movie Players",
-	"CreatedBy" : "Epic Games, Inc.",
-	"CreatedByURL" : "http://epicgames.com",
-	"DocsURL" : "",
-	"MarketplaceURL" : "",
-	"SupportURL" : "",
-	"EnabledByDefault" : true,
-	"CanContainContent" : false,
-	"IsBetaVersion" : false,
-	"Installed" : false,
-	"Modules" :
-	[
+	"FileVersion": 3,
+	"Version": 1,
+	"VersionName": "1.0",
+	"FriendlyName": "Apple Movie Player",
+	"Description": "Apple Platform Movie Player using AVPlayer library",
+	"Category": "Media Players",
+	"CreatedBy": "Epic Games, Inc.",
+	"CreatedByURL": "http://epicgames.com",
+	"DocsURL": "",
+	"MarketplaceURL": "",
+	"SupportURL": "",
+	"EnabledByDefault": true,
+	"CanContainContent": false,
+	"IsBetaVersion": false,
+	"Installed": false,
+	"Modules": [
 		{
-			"Name" : "AppleMoviePlayer",
-			"Type" : "RuntimeNoCommandlet",
-			"LoadingPhase" : "PreLoadingScreen",
-			"WhitelistPlatforms" :
-			[
+			"Name": "AppleMoviePlayer",
+			"Type": "RuntimeNoCommandlet",
+			"LoadingPhase": "PreLoadingScreen",
+			"WhitelistPlatforms": [
 				"IOS",
 				"TVOS",
 				"Mac"
 			]
 		}
-	]
+	],
+	"IsExperimentalVersion": false
 }
\ No newline at end of file
diff --git a/Engine/Plugins/Runtime/Firebase/Firebase.uplugin b/Engine/Plugins/Runtime/Firebase/Firebase.uplugin
--- a/Engine/Plugins/Runtime/Firebase/Firebase.uplugin
+++ b/Engine/Plugins/Runtime/Firebase/Firebase.uplugin
@@ -1,24 +1,33 @@
 {
-	"FileVersion" : 3,
-	"Version" : 1,
-	"VersionName" : "1.0",
-	"FriendlyName" : "Firebase",
-	"Description" : "Support for remote notifications using Firebase",
-	"Category" : "Notifications",
-	"CreatedBy" : "Epic Games, Inc.",
-	"CreatedByURL" : "http://epicgames.com",
-	"EnabledByDefault" : false,
-	"Installed" : false,
+	"FileVersion": 3,
+	"Version": 1,
+	"VersionName": "1.0",
+	"FriendlyName": "Firebase",
+	"Description": "Support for remote notifications using Firebase",
+	"Category": "Online Platform",
+	"CreatedBy": "Epic Games, Inc.",
+	"CreatedByURL": "http://epicgames.com",
+	"EnabledByDefault": false,
+	"Installed": false,
 	"SupportedTargetPlatforms": [
-		"Android", "IOS"
+		"Android",
+		"IOS"
 	],
-	"Modules" :
-	[
+	"Modules": [
 		{
 			"Name": "Firebase",
 			"Type": "Runtime",
 			"LoadingPhase": "PreDefault",
-			"WhitelistPlatforms": [ "Android", "IOS" ]
+			"WhitelistPlatforms": [
+				"Android",
+				"IOS"
+			]
 		}
-	]
-}
+	],
+	"DocsURL": "",
+	"MarketplaceURL": "",
+	"SupportURL": "",
+	"CanContainContent": false,
+	"IsBetaVersion": false,
+	"IsExperimentalVersion": false
+}
\ No newline at end of file
diff --git a/Engine/Plugins/Runtime/GoogleCloudMessaging/GoogleCloudMessaging.uplugin b/Engine/Plugins/Runtime/GoogleCloudMessaging/GoogleCloudMessaging.uplugin
--- a/Engine/Plugins/Runtime/GoogleCloudMessaging/GoogleCloudMessaging.uplugin
+++ b/Engine/Plugins/Runtime/GoogleCloudMessaging/GoogleCloudMessaging.uplugin
@@ -1,21 +1,28 @@
 {
-	"FileVersion" : 3,
-	"Version" : 1,
-	"VersionName" : "1.0",
-	"FriendlyName" : "Google Cloud Messaging",
-	"Description" : "Support for remote notifications using Google Cloud Messaging",
-	"Category" : "Notifications",
-	"CreatedBy" : "Epic Games, Inc.",
-	"CreatedByURL" : "http://epicgames.com",
-	"EnabledByDefault" : true,
-  "Installed" : false,
-	"Modules" :
-	[
-        {
-            "Name": "GoogleCloudMessaging",
-            "Type": "Runtime",
-            "LoadingPhase": "PreDefault",
-            "WhitelistPlatforms": [ "Android" ]
-        }
-	]
+	"FileVersion": 3,
+	"Version": 1,
+	"VersionName": "1.0",
+	"FriendlyName": "Google Cloud Messaging",
+	"Description": "Support for remote notifications using Google Cloud Messaging",
+	"Category": "Online Platform",
+	"CreatedBy": "Epic Games, Inc.",
+	"CreatedByURL": "http://epicgames.com",
+	"EnabledByDefault": true,
+	"Installed": false,
+	"Modules": [
+		{
+			"Name": "GoogleCloudMessaging",
+			"Type": "Runtime",
+			"LoadingPhase": "PreDefault",
+			"WhitelistPlatforms": [
+				"Android"
+			]
+		}
+	],
+	"DocsURL": "",
+	"MarketplaceURL": "",
+	"SupportURL": "",
+	"CanContainContent": false,
+	"IsBetaVersion": false,
+	"IsExperimentalVersion": false
 }
\ No newline at end of file
diff --git a/Engine/Plugins/Runtime/IOSReplayKit/IOSReplayKit.uplugin b/Engine/Plugins/Runtime/IOSReplayKit/IOSReplayKit.uplugin
--- a/Engine/Plugins/Runtime/IOSReplayKit/IOSReplayKit.uplugin
+++ b/Engine/Plugins/Runtime/IOSReplayKit/IOSReplayKit.uplugin
@@ -1,20 +1,25 @@
 {
-	"FileVersion" : 3,
-	"Version" : 1,
-	"VersionName" : "1.0",
-	"FriendlyName" : "ReplayKit for iOS",
-	"Description" : "Support for local recording and broadcasting using ReplayKit",
-	"Category" : "Replay",
-	"CreatedBy" : "Epic Games, Inc.",
-	"CreatedByURL" : "http://epicgames.com",
-	"EnabledByDefault" : false,
-    "Installed" : false,
-	"Modules" :
-	[
+	"FileVersion": 3,
+	"Version": 1,
+	"VersionName": "1.0",
+	"FriendlyName": "ReplayKit for iOS",
+	"Description": "Support for local recording and broadcasting using ReplayKit",
+	"Category": "Mobile",
+	"CreatedBy": "Epic Games, Inc.",
+	"CreatedByURL": "http://epicgames.com",
+	"EnabledByDefault": false,
+	"Installed": false,
+	"Modules": [
 		{
 			"Name": "IOSReplayKit",
 			"Type": "Runtime",
 			"LoadingPhase": "PreDefault"
 		}
-	]
-}
+	],
+	"DocsURL": "",
+	"MarketplaceURL": "",
+	"SupportURL": "",
+	"CanContainContent": false,
+	"IsBetaVersion": false,
+	"IsExperimentalVersion": false
+}
\ No newline at end of file
diff --git a/Engine/Plugins/Runtime/Interchange/Interchange.uplugin b/Engine/Plugins/Runtime/Interchange/Interchange.uplugin
--- a/Engine/Plugins/Runtime/Interchange/Interchange.uplugin
+++ b/Engine/Plugins/Runtime/Interchange/Interchange.uplugin
@@ -4,7 +4,7 @@
 	"VersionName": "1.0",
 	"FriendlyName": "Interchange Framework",
 	"Description": "Interchange framework plugin implement all implementation of engine importers and exporters.",
-	"Category": "Import Export",
+	"Category": "Importers",
 	"CreatedBy": "Epic Games, Inc.",
 	"CreatedByURL": "http://epicgames.com",
 	"DocsURL": "",
@@ -50,5 +50,6 @@
 				"Mac"
 			]
 		}
-	]
+	],
+	"IsExperimentalVersion": false
 }
\ No newline at end of file
diff --git a/Engine/Plugins/Runtime/LocationServicesAndroidImpl/LocationServicesAndroidImpl.uplugin b/Engine/Plugins/Runtime/LocationServicesAndroidImpl/LocationServicesAndroidImpl.uplugin
--- a/Engine/Plugins/Runtime/LocationServicesAndroidImpl/LocationServicesAndroidImpl.uplugin
+++ b/Engine/Plugins/Runtime/LocationServicesAndroidImpl/LocationServicesAndroidImpl.uplugin
@@ -1,30 +1,27 @@
 {
-	"FileVersion" : 3,
-	"Version" : 1,
-	"VersionName" : "1.0",
-	"FriendlyName" : "Mobile Location Services - Android Implementation",
-	"Description" : "Android implementation for blueprint access for location data from mobile devices",
-	"Category" : "Blueprints",
-	"CreatedBy" : "Epic Games, Inc.",
-	"CreatedByURL" : "http://epicgames.com",
-	"EnabledByDefault" : false,
-	"CanContainContent" : false,
-	"IsBetaVersion" : false,
-	"Installed" : false,
-	
-	"Modules" :
-	[
+	"FileVersion": 3,
+	"Version": 1,
+	"VersionName": "1.0",
+	"FriendlyName": "Mobile Location Services - Android Implementation",
+	"Description": "Android implementation for blueprint access for location data from mobile devices",
+	"Category": "Mobile",
+	"CreatedBy": "Epic Games, Inc.",
+	"CreatedByURL": "http://epicgames.com",
+	"EnabledByDefault": false,
+	"CanContainContent": false,
+	"IsBetaVersion": false,
+	"Installed": false,
+	"Modules": [
 		{
-			"Name" : "LocationServicesAndroidEditor",
-			"Type" : "Editor",
-            "LoadingPhase": "Default"
+			"Name": "LocationServicesAndroidEditor",
+			"Type": "Editor",
+			"LoadingPhase": "Default"
 		},
 		{
-			"Name" : "LocationServicesAndroidImpl",
-			"Type" : "Runtime",
-			"LoadingPhase" : "Default",
-			"WhitelistPlatforms" :
-			[
+			"Name": "LocationServicesAndroidImpl",
+			"Type": "Runtime",
+			"LoadingPhase": "Default",
+			"WhitelistPlatforms": [
 				"Android"
 			]
 		}
@@ -38,5 +35,9 @@
 			"Name": "AndroidPermission",
 			"Enabled": true
 		}
-	]
+	],
+	"DocsURL": "",
+	"MarketplaceURL": "",
+	"SupportURL": "",
+	"IsExperimentalVersion": false
 }
\ No newline at end of file
diff --git a/Engine/Plugins/Runtime/LocationServicesBPLibrary/LocationServicesBPLibrary.uplugin b/Engine/Plugins/Runtime/LocationServicesBPLibrary/LocationServicesBPLibrary.uplugin
--- a/Engine/Plugins/Runtime/LocationServicesBPLibrary/LocationServicesBPLibrary.uplugin
+++ b/Engine/Plugins/Runtime/LocationServicesBPLibrary/LocationServicesBPLibrary.uplugin
@@ -1,23 +1,25 @@
 {
-	"FileVersion" : 3,
-	"Version" : 1,
-	"VersionName" : "1.0",
-	"FriendlyName" : "Mobile Location Services Blueprints Library",
-	"Description" : "Common interface for blueprint access for location data from mobile devices",
-	"Category" : "Blueprints",
-	"CreatedBy" : "Epic Games, Inc.",
-	"CreatedByURL" : "http://epicgames.com",
-	"EnabledByDefault" : true,
-	"CanContainContent" : false,
-	"IsBetaVersion" : false,
-	"Installed" : false,
-	
-	"Modules" :
-	[
+	"FileVersion": 3,
+	"Version": 1,
+	"VersionName": "1.0",
+	"FriendlyName": "Mobile Location Services Blueprints Library",
+	"Description": "Common interface for blueprint access for location data from mobile devices",
+	"Category": "Mobile",
+	"CreatedBy": "Epic Games, Inc.",
+	"CreatedByURL": "http://epicgames.com",
+	"EnabledByDefault": true,
+	"CanContainContent": false,
+	"IsBetaVersion": false,
+	"Installed": false,
+	"Modules": [
 		{
-			"Name" : "LocationServicesBPLibrary",
-			"Type" : "Runtime",
-			"LoadingPhase" : "PreDefault"
+			"Name": "LocationServicesBPLibrary",
+			"Type": "Runtime",
+			"LoadingPhase": "PreDefault"
 		}
-	]
+	],
+	"DocsURL": "",
+	"MarketplaceURL": "",
+	"SupportURL": "",
+	"IsExperimentalVersion": false
 }
\ No newline at end of file
diff --git a/Engine/Plugins/Runtime/LocationServicesIOSImpl/LocationServicesIOSImpl.uplugin b/Engine/Plugins/Runtime/LocationServicesIOSImpl/LocationServicesIOSImpl.uplugin
--- a/Engine/Plugins/Runtime/LocationServicesIOSImpl/LocationServicesIOSImpl.uplugin
+++ b/Engine/Plugins/Runtime/LocationServicesIOSImpl/LocationServicesIOSImpl.uplugin
@@ -1,30 +1,27 @@
 {
-	"FileVersion" : 3,
-	"Version" : 1,
-	"VersionName" : "1.0",
-	"FriendlyName" : "Mobile Location Services - IOS Implementation",
-	"Description" : "IOS implementation for blueprint access for location data from mobile devices",
-	"Category" : "Blueprints",
-	"CreatedBy" : "Epic Games, Inc.",
-	"CreatedByURL" : "http://epicgames.com",
-	"EnabledByDefault" : false,
-	"CanContainContent" : false,
-	"IsBetaVersion" : false,
-	"Installed" : false,
-	
-	"Modules" :
-	[
+	"FileVersion": 3,
+	"Version": 1,
+	"VersionName": "1.0",
+	"FriendlyName": "Mobile Location Services - IOS Implementation",
+	"Description": "IOS implementation for blueprint access for location data from mobile devices",
+	"Category": "Mobile",
+	"CreatedBy": "Epic Games, Inc.",
+	"CreatedByURL": "http://epicgames.com",
+	"EnabledByDefault": false,
+	"CanContainContent": false,
+	"IsBetaVersion": false,
+	"Installed": false,
+	"Modules": [
 		{
-			"Name" : "LocationServicesIOSEditor",
-			"Type" : "Editor",
-            "LoadingPhase": "Default"
+			"Name": "LocationServicesIOSEditor",
+			"Type": "Editor",
+			"LoadingPhase": "Default"
 		},
 		{
-			"Name" : "LocationServicesIOSImpl",
-			"Type" : "Runtime",
-			"LoadingPhase" : "Default",
-			"WhitelistPlatforms" :
-			[
+			"Name": "LocationServicesIOSImpl",
+			"Type": "Runtime",
+			"LoadingPhase": "Default",
+			"WhitelistPlatforms": [
 				"IOS"
 			]
 		}
@@ -34,5 +31,9 @@
 			"Name": "LocationServicesBPLibrary",
 			"Enabled": true
 		}
-	]
+	],
+	"DocsURL": "",
+	"MarketplaceURL": "",
+	"SupportURL": "",
+	"IsExperimentalVersion": false
 }
\ No newline at end of file
diff --git a/Engine/Plugins/Runtime/OptionalMobileFeaturesBPLibrary/OptionalMobileFeaturesBPLibrary.uplugin b/Engine/Plugins/Runtime/OptionalMobileFeaturesBPLibrary/OptionalMobileFeaturesBPLibrary.uplugin
--- a/Engine/Plugins/Runtime/OptionalMobileFeaturesBPLibrary/OptionalMobileFeaturesBPLibrary.uplugin
+++ b/Engine/Plugins/Runtime/OptionalMobileFeaturesBPLibrary/OptionalMobileFeaturesBPLibrary.uplugin
@@ -1,25 +1,25 @@
 {
-	"FileVersion" : 3,
-	"Version" : 1,
-	"VersionName" : "1.0",
-	"FriendlyName" : "Optional Mobile Features Blueprint Library",
-	"Description" : "Gives blueprint access to Sound Volume, Battery Charge Level, and System Temperature for Android and iOS devices",
-	"Category" : "Blueprints",
-	"CreatedBy" : "Epic Games, Inc.",
-	"CreatedByURL" : "http://epicgames.com",
-	"DocsURL" : "",
-	"MarketplaceURL" : "",
-	"SupportURL" : "",
-	"EnabledByDefault" : false,
-	"CanContainContent" : false,
-	"IsBetaVersion" : false,
-	"Installed" : false,
-	"Modules" :
-	[
+	"FileVersion": 3,
+	"Version": 1,
+	"VersionName": "1.0",
+	"FriendlyName": "Optional Mobile Features Blueprint Library",
+	"Description": "Gives blueprint access to Sound Volume, Battery Charge Level, and System Temperature for Android and iOS devices",
+	"Category": "Mobile",
+	"CreatedBy": "Epic Games, Inc.",
+	"CreatedByURL": "http://epicgames.com",
+	"DocsURL": "",
+	"MarketplaceURL": "",
+	"SupportURL": "",
+	"EnabledByDefault": false,
+	"CanContainContent": false,
+	"IsBetaVersion": false,
+	"Installed": false,
+	"Modules": [
 		{
-			"Name" : "OptionalMobileFeaturesBPLibrary",
-			"Type" : "Runtime",
-			"LoadingPhase" : "PreDefault"
+			"Name": "OptionalMobileFeaturesBPLibrary",
+			"Type": "Runtime",
+			"LoadingPhase": "PreDefault"
 		}
-	]
+	],
+	"IsExperimentalVersion": false
 }
\ No newline at end of file
diff --git a/Engine/Plugins/Runtime/Steam/SteamShared/SteamShared.uplugin b/Engine/Plugins/Runtime/Steam/SteamShared/SteamShared.uplugin
--- a/Engine/Plugins/Runtime/Steam/SteamShared/SteamShared.uplugin
+++ b/Engine/Plugins/Runtime/Steam/SteamShared/SteamShared.uplugin
@@ -1,26 +1,30 @@
 {
-	"FileVersion" : 3,
-	"Version" : 1,
-	"VersionName" : "1.0",
-	"FriendlyName" : "Steam Shared Module",
-	"Description" : "Shared module loader for the Steam API",
-	"Category" : "Runtime",
-	"CreatedBy" : "Epic Games, Inc.",
-	"CreatedByURL" : "https://epicgames.com",
-	"EnabledByDefault" : false,
-	"CanContainContent" : false,
-	"Modules" :
-	[
+	"FileVersion": 3,
+	"Version": 1,
+	"VersionName": "1.0",
+	"FriendlyName": "Steam Shared Module",
+	"Description": "Shared module loader for the Steam API",
+	"Category": "Online Platform",
+	"CreatedBy": "Epic Games, Inc.",
+	"CreatedByURL": "https://epicgames.com",
+	"EnabledByDefault": false,
+	"CanContainContent": false,
+	"Modules": [
 		{
-			"Name" : "SteamShared",
-			"Type" : "Runtime",
-			"LoadingPhase" : "PreDefault",
-			"WhitelistPlatforms" :
-			[
-                "Win64",
-                "Mac",
-                "Linux"
+			"Name": "SteamShared",
+			"Type": "Runtime",
+			"LoadingPhase": "PreDefault",
+			"WhitelistPlatforms": [
+				"Win64",
+				"Mac",
+				"Linux"
 			]
 		}
-	]
+	],
+	"DocsURL": "",
+	"MarketplaceURL": "",
+	"SupportURL": "",
+	"IsBetaVersion": false,
+	"IsExperimentalVersion": false,
+	"Installed": false
 }
\ No newline at end of file
diff --git a/Engine/Plugins/Runtime/WebBrowserNativeProxy/WebBrowserNativeProxy.uplugin b/Engine/Plugins/Runtime/WebBrowserNativeProxy/WebBrowserNativeProxy.uplugin
--- a/Engine/Plugins/Runtime/WebBrowserNativeProxy/WebBrowserNativeProxy.uplugin
+++ b/Engine/Plugins/Runtime/WebBrowserNativeProxy/WebBrowserNativeProxy.uplugin
@@ -1,31 +1,31 @@
 {
-	"FileVersion" : 3,
-	"Version" : 1,
-	"VersionName" : "1.0",
-	"FriendlyName" : "Web Browser to Native Proxying",
-	"Description" : "Maintains the browser to native proxy and provides hooks for registering UObjects bindings",
-	"Category" : "Web",
-	"CreatedBy" : "Epic Games, Inc.",
-	"CreatedByURL" : "http://epicgames.com",
-	"DocsURL" : "",
-	"MarketplaceURL" : "",
-	"SupportURL" : "",
-	"EnabledByDefault" : false,
-	"IsBetaVersion" : false,
-	"Installed" : false,
-	"Modules" :
-	[
+	"FileVersion": 3,
+	"Version": 1,
+	"VersionName": "1.0",
+	"FriendlyName": "Web Browser to Native Proxying",
+	"Description": "Maintains the browser to native proxy and provides hooks for registering UObjects bindings",
+	"Category": "UI",
+	"CreatedBy": "Epic Games, Inc.",
+	"CreatedByURL": "http://epicgames.com",
+	"DocsURL": "",
+	"MarketplaceURL": "",
+	"SupportURL": "",
+	"EnabledByDefault": false,
+	"IsBetaVersion": false,
+	"Installed": false,
+	"Modules": [
 		{
-			"Name" : "WebBrowserNativeProxy",
-			"Type" : "Runtime",
-			"LoadingPhase" : "PreDefault",
-			"WhitelistPlatforms" :
-			[
+			"Name": "WebBrowserNativeProxy",
+			"Type": "Runtime",
+			"LoadingPhase": "PreDefault",
+			"WhitelistPlatforms": [
 				"Win64",
 				"Mac",
 				"iOS",
 				"Android"
 			]
 		}
-	]
-}
+	],
+	"CanContainContent": false,
+	"IsExperimentalVersion": false
+}
\ No newline at end of file
diff --git a/Engine/Plugins/Runtime/WebBrowserWidget/WebBrowserWidget.uplugin b/Engine/Plugins/Runtime/WebBrowserWidget/WebBrowserWidget.uplugin
--- a/Engine/Plugins/Runtime/WebBrowserWidget/WebBrowserWidget.uplugin
+++ b/Engine/Plugins/Runtime/WebBrowserWidget/WebBrowserWidget.uplugin
@@ -1,27 +1,25 @@
 {
-	"FileVersion" : 3,
-	"Version" : 1,
-	"VersionName" : "1.0",
-	"FriendlyName" : "Web Browser",
-	"Description" : "Allows the user to create a Web Browser Widget",
-	"Category" : "Widgets",
-	"CreatedBy" : "Epic Games, Inc.",
-	"CreatedByURL" : "http://epicgames.com",
-	"DocsURL" : "",
-	"MarketplaceURL" : "",
-	"SupportURL" : "",
-	"EnabledByDefault" : false,
-	"CanContainContent" : true,
-	"IsBetaVersion" : false,
-	"Installed" : false,
-	"Modules" :
-	[
+	"FileVersion": 3,
+	"Version": 1,
+	"VersionName": "1.0",
+	"FriendlyName": "Web Browser",
+	"Description": "Allows the user to create a Web Browser Widget",
+	"Category": "UI",
+	"CreatedBy": "Epic Games, Inc.",
+	"CreatedByURL": "http://epicgames.com",
+	"DocsURL": "",
+	"MarketplaceURL": "",
+	"SupportURL": "",
+	"EnabledByDefault": false,
+	"CanContainContent": true,
+	"IsBetaVersion": false,
+	"Installed": false,
+	"Modules": [
 		{
-			"Name" : "WebBrowserWidget",
-			"Type" : "Runtime",
-			"LoadingPhase" : "PreDefault",
-			"WhitelistPlatforms" :
-			[
+			"Name": "WebBrowserWidget",
+			"Type": "Runtime",
+			"LoadingPhase": "PreDefault",
+			"WhitelistPlatforms": [
 				"Win64",
 				"Mac",
 				"Linux",
@@ -30,5 +28,6 @@
 				"PS4"
 			]
 		}
-	]
-}
+	],
+	"IsExperimentalVersion": false
+}
\ No newline at end of file
diff --git a/Engine/Plugins/Runtime/WebMMoviePlayer/WebMMoviePlayer.uplugin b/Engine/Plugins/Runtime/WebMMoviePlayer/WebMMoviePlayer.uplugin
--- a/Engine/Plugins/Runtime/WebMMoviePlayer/WebMMoviePlayer.uplugin
+++ b/Engine/Plugins/Runtime/WebMMoviePlayer/WebMMoviePlayer.uplugin
@@ -4,7 +4,7 @@
 	"VersionName": "1.0",
 	"FriendlyName": "Movie Player for WebM files",
 	"Description": "Movie Player for WebM files",
-	"Category": "Movie Players",
+	"Category": "Media Players",
 	"CreatedBy": "Epic Games, Inc.",
 	"CreatedByURL": "http://epicgames.com",
 	"DocsURL": "",
@@ -34,5 +34,6 @@
 			"Name": "WebMMedia",
 			"Enabled": true
 		}
-	]
-}
+	],
+	"IsExperimentalVersion": false
+}
\ No newline at end of file
diff --git a/Engine/Plugins/Runtime/WindowsMoviePlayer/WindowsMoviePlayer.uplugin b/Engine/Plugins/Runtime/WindowsMoviePlayer/WindowsMoviePlayer.uplugin
--- a/Engine/Plugins/Runtime/WindowsMoviePlayer/WindowsMoviePlayer.uplugin
+++ b/Engine/Plugins/Runtime/WindowsMoviePlayer/WindowsMoviePlayer.uplugin
@@ -1,29 +1,28 @@
 {
-	"FileVersion" : 3,
-	"Version" : 1,
-	"VersionName" : "1.0",
-	"FriendlyName" : "Windows Movie Player",
-	"Description" : "Windows Specific Movie Player using Media Foundation",
-	"Category" : "Movie Players",
-	"CreatedBy" : "Epic Games, Inc.",
-	"CreatedByURL" : "http://epicgames.com",
-	"DocsURL" : "",
-	"MarketplaceURL" : "",
-	"SupportURL" : "",
-	"EnabledByDefault" : true,
-	"CanContainContent" : false,
-	"IsBetaVersion" : false,
-	"Installed" : false,
-	"Modules" :
-	[
+	"FileVersion": 3,
+	"Version": 1,
+	"VersionName": "1.0",
+	"FriendlyName": "Windows Movie Player",
+	"Description": "Windows Specific Movie Player using Media Foundation",
+	"Category": "Media Players",
+	"CreatedBy": "Epic Games, Inc.",
+	"CreatedByURL": "http://epicgames.com",
+	"DocsURL": "",
+	"MarketplaceURL": "",
+	"SupportURL": "",
+	"EnabledByDefault": true,
+	"CanContainContent": false,
+	"IsBetaVersion": false,
+	"Installed": false,
+	"Modules": [
 		{
-			"Name" : "WindowsMoviePlayer",
-			"Type" : "RuntimeNoCommandlet",
-			"LoadingPhase" : "PreLoadingScreen",
-			"WhitelistPlatforms" :
-			[
+			"Name": "WindowsMoviePlayer",
+			"Type": "RuntimeNoCommandlet",
+			"LoadingPhase": "PreLoadingScreen",
+			"WhitelistPlatforms": [
 				"Win64"
 			]
 		}
-	]
+	],
+	"IsExperimentalVersion": false
 }
\ No newline at end of file
diff --git a/Engine/Plugins/VirtualProduction/DMX/DMXFixtures/DMXFixtures.uplugin b/Engine/Plugins/VirtualProduction/DMX/DMXFixtures/DMXFixtures.uplugin
--- a/Engine/Plugins/VirtualProduction/DMX/DMXFixtures/DMXFixtures.uplugin
+++ b/Engine/Plugins/VirtualProduction/DMX/DMXFixtures/DMXFixtures.uplugin
@@ -1,24 +1,24 @@
 {
-  "FileVersion": 1,
-  "Version": 1,
-  "VersionName": "1.0",
-  "FriendlyName": "DMX Fixtures",
-  "Description": "DMX Light Fixtures Blueprints",
-  "Category": "Media",
-  "CreatedBy": "Epic Games, Inc.",
-  "CreatedByURL": "http://epicgames.com",
-  "DocsURL": "",
-  "MarketplaceURL": "",
-  "SupportURL": "",
-  "CanContainContent": true,
-  "IsBetaVersion": true,
-  "Modules": [
-    {
-      "Name": "DMXFixtures",
-      "Type": "Runtime",
-      "LoadingPhase": "Default"
-    }
-   ],
+	"FileVersion": 3,
+	"Version": 1,
+	"VersionName": "1.0",
+	"FriendlyName": "DMX Fixtures",
+	"Description": "DMX Light Fixtures Blueprints",
+	"Category": "Virtual Production",
+	"CreatedBy": "Epic Games, Inc.",
+	"CreatedByURL": "http://epicgames.com",
+	"DocsURL": "",
+	"MarketplaceURL": "",
+	"SupportURL": "",
+	"CanContainContent": true,
+	"IsBetaVersion": true,
+	"Modules": [
+		{
+			"Name": "DMXFixtures",
+			"Type": "Runtime",
+			"LoadingPhase": "Default"
+		}
+	],
 	"Plugins": [
 		{
 			"Name": "DMXEngine",
@@ -32,5 +32,7 @@
 			"Name": "ProceduralMeshComponent",
 			"Enabled": true
 		}
-	]
+	],
+	"IsExperimentalVersion": false,
+	"Installed": false
 }
\ No newline at end of file
