Elimated use of UMG style in DMXPixelMapping plugin. Style is deprecated and unnecessary.

diff --git a/Engine/Plugins/Experimental/DigitalMultiplex/DMXPixelMapping/Source/DMXPixelMappingEditor/Private/Customizations/DMXPixelMappingDetailCustomization_FixtureGroup.cpp b/Engine/Plugins/Experimental/DigitalMultiplex/DMXPixelMapping/Source/DMXPixelMappingEditor/Private/Customizations/DMXPixelMappingDetailCustomization_FixtureGroup.cpp
--- a/Engine/Plugins/Experimental/DigitalMultiplex/DMXPixelMapping/Source/DMXPixelMappingEditor/Private/Customizations/DMXPixelMappingDetailCustomization_FixtureGroup.cpp
+++ b/Engine/Plugins/Experimental/DigitalMultiplex/DMXPixelMapping/Source/DMXPixelMappingEditor/Private/Customizations/DMXPixelMappingDetailCustomization_FixtureGroup.cpp
@@ -54,7 +54,6 @@ TSharedRef<ITableRow> FDMXPixelMappingDetailCustomization_FixtureGroup::Generate
 			return SNew(STableRow<TSharedPtr<FString>>, OwnerTable)
 				.Padding(2.0f)
 				.OnDragDetected(FOnDragDetected::CreateSP(this, &FDMXPixelMappingDetailCustomization_FixtureGroup::OnFixturePatchDragDetected, InFixturePatchRef))
-				.Style(FEditorStyle::Get(), "UMGEditor.PaletteItem")
 				.ShowSelection(true)
 				[
 					SNew(SBox)
diff --git a/Engine/Plugins/Experimental/DigitalMultiplex/DMXPixelMapping/Source/DMXPixelMappingEditor/Private/Customizations/DMXPixelMappingDetailCustomization_FixtureGroupItem.cpp b/Engine/Plugins/Experimental/DigitalMultiplex/DMXPixelMapping/Source/DMXPixelMappingEditor/Private/Customizations/DMXPixelMappingDetailCustomization_FixtureGroupItem.cpp
--- a/Engine/Plugins/Experimental/DigitalMultiplex/DMXPixelMapping/Source/DMXPixelMappingEditor/Private/Customizations/DMXPixelMappingDetailCustomization_FixtureGroupItem.cpp
+++ b/Engine/Plugins/Experimental/DigitalMultiplex/DMXPixelMapping/Source/DMXPixelMappingEditor/Private/Customizations/DMXPixelMappingDetailCustomization_FixtureGroupItem.cpp
@@ -195,7 +195,6 @@ TSharedRef<ITableRow> FDMXPixelMappingDetailCustomization_FixtureGroupItem::Gene
 
 	return SNew(STableRow<TSharedPtr<FString>>, OwnerTable)
 		.Padding(2.0f)
-		.Style(FEditorStyle::Get(), "UMGEditor.PaletteItem")
 		.ShowSelection(false)
 		[
 			SNew(SBox)
diff --git a/Engine/Plugins/Experimental/DigitalMultiplex/DMXPixelMapping/Source/DMXPixelMappingEditor/Private/Customizations/DMXPixelMappingDetailCustomization_Matrix.cpp b/Engine/Plugins/Experimental/DigitalMultiplex/DMXPixelMapping/Source/DMXPixelMappingEditor/Private/Customizations/DMXPixelMappingDetailCustomization_Matrix.cpp
--- a/Engine/Plugins/Experimental/DigitalMultiplex/DMXPixelMapping/Source/DMXPixelMappingEditor/Private/Customizations/DMXPixelMappingDetailCustomization_Matrix.cpp
+++ b/Engine/Plugins/Experimental/DigitalMultiplex/DMXPixelMapping/Source/DMXPixelMappingEditor/Private/Customizations/DMXPixelMappingDetailCustomization_Matrix.cpp
@@ -219,7 +219,6 @@ TSharedRef<ITableRow> FDMXPixelMappingDetailCustomization_Matrix::GenerateExpose
 
 	return SNew(STableRow<TSharedPtr<FString>>, OwnerTable)
 		.Padding(2.0f)
-		.Style(FEditorStyle::Get(), "UMGEditor.PaletteItem")
 		.ShowSelection(false)
 		[
 			SNew(SBox)
