DMX - Fix PixelMapping does not output DMX in non-editor builds when a Material or UMG widget is used

diff --git a/Engine/Plugins/VirtualProduction/DMX/DMXPixelMapping/Source/DMXPixelMappingRuntime/Private/Components/DMXPixelMappingRendererComponent.cpp b/Engine/Plugins/VirtualProduction/DMX/DMXPixelMapping/Source/DMXPixelMappingRuntime/Private/Components/DMXPixelMappingRendererComponent.cpp
--- a/Engine/Plugins/VirtualProduction/DMX/DMXPixelMapping/Source/DMXPixelMappingRuntime/Private/Components/DMXPixelMappingRendererComponent.cpp
+++ b/Engine/Plugins/VirtualProduction/DMX/DMXPixelMapping/Source/DMXPixelMappingRuntime/Private/Components/DMXPixelMappingRendererComponent.cpp
@@ -478,24 +478,34 @@ void UDMXPixelMappingRendererComponent::RendererInputTexture()
 {
 	Initialize();
 
-#if WITH_EDITOR
 	switch (RendererType)
 	{
-	case(EDMXPixelMappingRendererType::Texture) :
-		// Nothing
-		if (InputTexture != nullptr &&  InputTexture->GetResource() != nullptr)
-		{
-			ResizePreviewRenderTarget(InputTexture->GetResource()->GetSizeX(), InputTexture->GetResource()->GetSizeY());
-		}
+	case(EDMXPixelMappingRendererType::Texture):
+		// No rendering required
 		break;
+
 	case(EDMXPixelMappingRendererType::Material):
 		PixelMappingRenderer->RenderMaterial(InputRenderTarget, InputMaterial);
-		ResizePreviewRenderTarget(SizeX, SizeY);
 		break;
+
 	case(EDMXPixelMappingRendererType::UMG):
 		PixelMappingRenderer->RenderWidget(InputRenderTarget, UserWidget);
-		ResizePreviewRenderTarget(SizeX, SizeY);
 		break;
+	default:
+		checkNoEntry(); // Unknown renderer type
+	}
+
+#if WITH_EDITOR
+	if (RendererType == EDMXPixelMappingRendererType::Texture)
+	{
+		if (InputTexture != nullptr && InputTexture->GetResource())
+		{
+			ResizePreviewRenderTarget(InputTexture->GetResource()->GetSizeX(), InputTexture->GetResource()->GetSizeY());
+		}
+	}
+	else
+	{
+		ResizePreviewRenderTarget(SizeX, SizeY);
 	}
 #endif
 }
