DMX - Rework DMX Fixture Type editor for better GDTF support, better UX and to ease maintaining and extending the editor

diff --git a/Engine/Plugins/VirtualProduction/DMX/DMXPixelMapping/Source/DMXPixelMappingRuntime/Private/Components/DMXPixelMappingMatrixComponent.cpp b/Engine/Plugins/VirtualProduction/DMX/DMXPixelMapping/Source/DMXPixelMappingRuntime/Private/Components/DMXPixelMappingMatrixComponent.cpp
--- a/Engine/Plugins/VirtualProduction/DMX/DMXPixelMapping/Source/DMXPixelMappingRuntime/Private/Components/DMXPixelMappingMatrixComponent.cpp
+++ b/Engine/Plugins/VirtualProduction/DMX/DMXPixelMapping/Source/DMXPixelMappingRuntime/Private/Components/DMXPixelMappingMatrixComponent.cpp
@@ -8,7 +8,6 @@
 #include "DMXPixelMappingRuntimeUtils.h"
 #include "DMXPixelMappingTypes.h"
 #include "DMXSubsystem.h"
-#include "DMXUtils.h"
 #include "Components/DMXPixelMappingFixtureGroupComponent.h"
 #include "Components/DMXPixelMappingMatrixCellComponent.h"
 #include "Components/DMXPixelMappingRendererComponent.h"
@@ -47,6 +46,9 @@ UDMXPixelMappingMatrixComponent::UDMXPixelMappingMatrixComponent()
 	AttributeBExpose = true;
 
 	bMonochromeExpose = true;
+
+	// Listen to Fixture Type changes
+	UDMXEntityFixtureType::GetOnFixtureTypeChanged().AddUObject(this, &UDMXPixelMappingMatrixComponent::OnFixtureTypeChanged);
 }
 
 void UDMXPixelMappingMatrixComponent::Serialize(FArchive& Ar)
@@ -139,22 +141,15 @@ void UDMXPixelMappingMatrixComponent::PostEditChangeProperty(FPropertyChangedEve
 	if (PropertyChangedEvent.GetPropertyName() == GET_MEMBER_NAME_CHECKED(UDMXPixelMappingMatrixComponent, FixturePatchRef) ||
 		PropertyChangedEvent.GetPropertyName() == GET_MEMBER_NAME_CHECKED(UDMXPixelMappingMatrixComponent, CoordinateGrid))
 	{
-		if (UDMXPixelMapping* PixelMapping = GetPixelMapping())
-		{
-			HandleMatrixChanged();
-		}
+		HandleMatrixChanged();
 	}
 	else if (PropertyChangedEvent.GetPropertyName() == GET_MEMBER_NAME_CHECKED(FDMXEntityReference, DMXLibrary))
 	{
-		if (UDMXPixelMapping* PixelMapping = GetPixelMapping())
-		{
-			HandleMatrixChanged();
-		}
+		HandleMatrixChanged();
 	}
 }
 #endif // WITH_EDITOR
 
-
 #if WITH_EDITOR
 void UDMXPixelMappingMatrixComponent::PostEditChangeChainProperty(FPropertyChangedChainEvent& PropertyChangedChainEvent)
 {
@@ -333,62 +328,6 @@ FString UDMXPixelMappingMatrixComponent::GetUserFriendlyName() const
 	return FString(TEXT("Fixture Matrix: No Fixture Patch"));
 }
 
-void UDMXPixelMappingMatrixComponent::Tick(float DeltaTime)
-{
-#if WITH_EDITOR
-	// Test for property changes each tick
-	if (UDMXPixelMapping * PixelMapping = GetPixelMapping())
-	{
-			bool bShouldRebuildChildren = false;
-
-			UDMXLibrary* DMXLibrary = FixturePatchRef.DMXLibrary;
-			UDMXEntityFixturePatch* FixturePatch = FixturePatchRef.GetFixturePatch();
-
-			if (DMXLibrary && FixturePatch)
-			{
-				if (UDMXEntityFixtureType * ParentFixtureType = FixturePatch->GetFixtureType())
-				{
-					if (!FixturePatch->GetActiveMode() && GetChildrenCount() > 0)
-					{
-						bShouldRebuildChildren = true;
-					}
-					else if (!ParentFixtureType->bFixtureMatrixEnabled && GetChildrenCount() > 0)
-					{
-						bShouldRebuildChildren = true;
-					}
-					else
-					{
-						const FDMXFixtureMode* FixtureMode = FixturePatch->GetActiveMode();
-
-						if (FixtureMode)
-						{
-							const FDMXFixtureMatrix& FixtureMatrixConfig = FixtureMode->FixtureMatrixConfig;
-						const FIntPoint NewCoordinateGrid(FixtureMatrixConfig.XCells, FixtureMatrixConfig.YCells);
-
-						if (CoordinateGrid != NewCoordinateGrid)
-							{
-								bShouldRebuildChildren = true;
-							}
-							else if (FixtureMatrixConfig.PixelMappingDistribution != Distribution)
-							{
-								bShouldRebuildChildren = true;
-								Distribution = FixtureMatrixConfig.PixelMappingDistribution;
-							}
-						}
-					}
-				}
-			}
-
-			if (bShouldRebuildChildren)
-			{
-			HandleMatrixChanged();
-
-				LogInvalidProperties();
-		}
-	}
-#endif // WITH_EDITOR
-}
-
 void UDMXPixelMappingMatrixComponent::SetPosition(const FVector2D& NewPosition)
 {
 	Modify();
@@ -474,43 +413,51 @@ void UDMXPixelMappingMatrixComponent::HandleMatrixChanged()
 
 	// Remove all existing children and rebuild them anew
 	for (UDMXPixelMappingBaseComponent* Child : TArray<UDMXPixelMappingBaseComponent*>(Children))
-{
+	{
 		RemoveChild(Child);
 	}
 
-	if (UDMXEntityFixturePatch* FixturePatch = FixturePatchRef.GetFixturePatch())
+
+	UDMXEntityFixturePatch* FixturePatch = FixturePatchRef.GetFixturePatch();
+	UDMXEntityFixtureType* FixtureType = FixturePatch ? FixturePatch->GetFixtureType() : nullptr;
+	const FDMXFixtureMode* ModePtr = FixturePatch ? FixturePatch->GetActiveMode() : nullptr;
+	if (FixturePatch && FixtureType && ModePtr && ModePtr->bFixtureMatrixEnabled)
 	{
-		if (UDMXEntityFixtureType* FixtureType = FixturePatch->GetFixtureType())
+		TArray<FDMXCell> MatrixCells;
+		if (FixturePatch->GetAllMatrixCells(MatrixCells))
 		{
-			if (FixtureType->bFixtureMatrixEnabled)
-			{
-				if (const FDMXFixtureMode* ModePtr = FixturePatch->GetActiveMode())
-				{
-					TArray<FDMXCell> MatrixCells;
-					if (FixturePatch->GetAllMatrixCells(MatrixCells))
-					{
-						Distribution = ModePtr->FixtureMatrixConfig.PixelMappingDistribution;
-						CoordinateGrid = FIntPoint(ModePtr->FixtureMatrixConfig.XCells, ModePtr->FixtureMatrixConfig.YCells);
+			Distribution = ModePtr->FixtureMatrixConfig.PixelMappingDistribution;
+			CoordinateGrid = FIntPoint(ModePtr->FixtureMatrixConfig.XCells, ModePtr->FixtureMatrixConfig.YCells);
 
-						for (const FDMXCell& Cell : MatrixCells)
-						{
-							TSharedPtr<FDMXPixelMappingComponentTemplate> ComponentTemplate = MakeShared<FDMXPixelMappingComponentTemplate>(UDMXPixelMappingMatrixCellComponent::StaticClass());
-							UDMXPixelMappingMatrixCellComponent* Component = ComponentTemplate->CreateComponent<UDMXPixelMappingMatrixCellComponent>(GetPixelMapping()->GetRootComponent());
+			for (const FDMXCell& Cell : MatrixCells)
+			{
+				TSharedPtr<FDMXPixelMappingComponentTemplate> ComponentTemplate = MakeShared<FDMXPixelMappingComponentTemplate>(UDMXPixelMappingMatrixCellComponent::StaticClass());
+				UDMXPixelMappingMatrixCellComponent* Component = ComponentTemplate->CreateComponent<UDMXPixelMappingMatrixCellComponent>(GetPixelMapping()->GetRootComponent());
 
-							Component->CellID = Cell.CellID;
-							Component->SetCellCoordinate(Cell.Coordinate);
+				Component->CellID = Cell.CellID;
+				Component->SetCellCoordinate(Cell.Coordinate);
 
-							AddChild(Component);
-				}
+				AddChild(Component);
 			}
 		}
 	}
-		}
-	}
 
 	HandleSizeChanged();
 
 	GetOnMatrixChanged().Broadcast(GetPixelMapping(), this);
 }
 
+void UDMXPixelMappingMatrixComponent::OnFixtureTypeChanged(const UDMXEntityFixtureType* FixtureType)
+{
+	if (UDMXEntityFixturePatch* FixturePatch = FixturePatchRef.GetFixturePatch())
+	{
+		if (FixturePatch->GetFixtureType() == FixtureType)
+		{
+			HandleMatrixChanged();
+
+			LogInvalidProperties();
+		}
+	}
+}
+
 #undef LOCTEXT_NAMESPACE
diff --git a/Engine/Plugins/VirtualProduction/DMX/DMXPixelMapping/Source/DMXPixelMappingRuntime/Public/Components/DMXPixelMappingMatrixComponent.h b/Engine/Plugins/VirtualProduction/DMX/DMXPixelMapping/Source/DMXPixelMappingRuntime/Public/Components/DMXPixelMappingMatrixComponent.h
--- a/Engine/Plugins/VirtualProduction/DMX/DMXPixelMapping/Source/DMXPixelMappingRuntime/Public/Components/DMXPixelMappingMatrixComponent.h
+++ b/Engine/Plugins/VirtualProduction/DMX/DMXPixelMapping/Source/DMXPixelMappingRuntime/Public/Components/DMXPixelMappingMatrixComponent.h
@@ -60,11 +60,6 @@ public:
 	virtual FString GetUserFriendlyName() const override;
 	// ~End UDMXPixelMappingBaseComponent interface
 
-	// ~Begin FTickableGameObject interface
-	virtual void Tick(float DeltaTime) override;
-	virtual bool IsTickable() const { return true; }
-	// ~End FTickableGameObject interface
-
 	// ~Begin UDMXPixelMappingOutputComponent interface
 #if WITH_EDITOR
 	virtual const FText GetPaletteCategory() override;
@@ -85,6 +80,9 @@ public:
 	void HandleMatrixChanged();
 
 protected:
+	/** Called when the fixture type in use changed */
+	void OnFixtureTypeChanged(const UDMXEntityFixtureType* FixtureType);
+
 #if WITH_EDITORONLY_DATA
 	/** Children available PreEditUndo, useful to hide all removed ones in post edit undo */
 	TArray<UDMXPixelMappingBaseComponent*> PreEditUndoMatrixCellChildren;
diff --git a/Engine/Plugins/VirtualProduction/DMX/DMXProtocol/Source/DMXProtocol/Private/DMXConversions.cpp b/Engine/Plugins/VirtualProduction/DMX/DMXProtocol/Source/DMXProtocol/Private/DMXConversions.cpp
--- a/Engine/Plugins/VirtualProduction/DMX/DMXProtocol/Source/DMXProtocol/Private/DMXConversions.cpp
+++ b/Engine/Plugins/VirtualProduction/DMX/DMXProtocol/Source/DMXProtocol/Private/DMXConversions.cpp
@@ -3,33 +3,31 @@
 #include "DMXConversions.h"
 
 
-TArray<uint8> FDMXConversions::UnsignedInt32ToByteArray(uint32 Value, EDMXFixtureSignalFormat SignalFormat, bool bLSBOrder)
+uint8 FDMXConversions::GetSizeOfSignalFormat(EDMXFixtureSignalFormat SignalFormat)
 {
-	const int32 NumBytes = static_cast<uint8>(SignalFormat) + 1;
-
-	// To avoid branching in the loop, we'll decide before it on which byte to start
-	// and which direction to go, depending on the Function's endianness.
-	const int8 ByteIndexStep = bLSBOrder ? 1 : -1;
-	int8 OutByteIndex = bLSBOrder ? 0 : NumBytes - 1;
+	return static_cast<uint8>(SignalFormat) + 1;
+}
 
-	TArray<uint8> Bytes;
-	Bytes.AddUninitialized(NumBytes);
-	for (uint8 ValueByte = 0; ValueByte < NumBytes; ++ValueByte)
+uint32 FDMXConversions::GetSignalFormatMaxValue(EDMXFixtureSignalFormat SignalFormat)
+{
+	switch (SignalFormat)
 	{
-		Bytes[OutByteIndex] = (Value >> 8 * ValueByte) & 0xFF;
-		OutByteIndex += ByteIndexStep;
+	case EDMXFixtureSignalFormat::E8Bit:
+		return 0x000000FF;
+	case EDMXFixtureSignalFormat::E16Bit:
+		return 0x0000FFFF;
+	case EDMXFixtureSignalFormat::E24Bit:
+		return 0x00FFFFFF;
+	case EDMXFixtureSignalFormat::E32Bit:
+		return 0xFFFFFFFF;
+	default:
+		checkNoEntry(); // Unhandled enum value
 	}
 
-	return Bytes;
+	return 0;
 }
 
-TArray<uint8> FDMXConversions::NormalizedDMXValueToByteArray(float NormalizedValue, EDMXFixtureSignalFormat SignalFormat, bool bLSBOrder)
+uint32 FDMXConversions::ClampValueBySignalFormat(uint32 Value, EDMXFixtureSignalFormat SignalFormat)
 {
-	NormalizedValue = FMath::Clamp(NormalizedValue, 0.f, 1.f);
-	uint32 Value = FMath::Floor<uint32>(TNumericLimits<uint32>::Max() * static_cast<double>(NormalizedValue));
-		
-	// Shift the value into signal format range
-	Value = Value >> ((3 - static_cast<uint8>(SignalFormat)) * 8);
-
-	return UnsignedInt32ToByteArray(Value, SignalFormat, bLSBOrder);
+	return FMath::Min(Value, GetSignalFormatMaxValue(SignalFormat));
 }
diff --git a/Engine/Plugins/VirtualProduction/DMX/DMXProtocol/Source/DMXProtocol/Private/DMXProtocolModule.cpp b/Engine/Plugins/VirtualProduction/DMX/DMXProtocol/Source/DMXProtocol/Private/DMXProtocolModule.cpp
--- a/Engine/Plugins/VirtualProduction/DMX/DMXProtocol/Source/DMXProtocol/Private/DMXProtocolModule.cpp
+++ b/Engine/Plugins/VirtualProduction/DMX/DMXProtocol/Source/DMXProtocol/Private/DMXProtocolModule.cpp
@@ -105,8 +105,8 @@ void FDMXProtocolModule::StartupModule()
 	// Register DMX Protocol global settings
 	if (SettingsModule != nullptr)
 	{
-		SettingsModule->RegisterSettings("Project", "Plugins", "DMX Plugin",
-			LOCTEXT("ProjectSettings_Label", "DMX Plugin"),
+		SettingsModule->RegisterSettings("Project", "Plugins", "DMX",
+			LOCTEXT("ProjectSettings_Label", "DMX"),
 			LOCTEXT("ProjectSettings_Description", "Configure DMX plugin global settings"),
 			GetMutableDefault<UDMXProtocolSettings>()
 		);
diff --git a/Engine/Plugins/VirtualProduction/DMX/DMXProtocol/Source/DMXProtocol/Public/DMXConversions.h b/Engine/Plugins/VirtualProduction/DMX/DMXProtocol/Source/DMXProtocol/Public/DMXConversions.h
--- a/Engine/Plugins/VirtualProduction/DMX/DMXProtocol/Source/DMXProtocol/Public/DMXConversions.h
+++ b/Engine/Plugins/VirtualProduction/DMX/DMXProtocol/Source/DMXProtocol/Public/DMXConversions.h
@@ -7,6 +7,8 @@
 #include "DMXProtocolTypes.h"
 
 #include "Misc/ByteSwap.h"
+#include "Templates/EnableIf.h"
+#include "Templates/UnrealTypeTraits.h"
 
 
 /** Common conversions for dmx */
@@ -18,10 +20,55 @@ class DMXPROTOCOL_API FDMXConversions
 
 public:
 	/** Converts a uint32 to a byte array */
-	static TArray<uint8> UnsignedInt32ToByteArray(uint32 Value, EDMXFixtureSignalFormat SignalFormat, bool bLSBOrder);
+	template <typename T, typename U = typename TEnableIf<TIsSame<T, uint32>::Value>::Type>
+	static TArray<uint8> UnsignedInt32ToByteArray(T Value, EDMXFixtureSignalFormat SignalFormat, bool bLSBOrder)
+	{
+		const uint8 NumBytes = GetSizeOfSignalFormat(SignalFormat);
+
+		// To avoid branching in the loop, we'll decide before it on which byte to start
+		// and which direction to go, depending on the Function's endianness.
+		const uint8 ByteIndexStep = bLSBOrder ? 1 : -1;
+		uint8 OutByteIndex = bLSBOrder ? 0 : NumBytes - 1;
+
+		TArray<uint8> Bytes;
+		Bytes.AddUninitialized(NumBytes);
+		for (uint8 ValueByte = 0; ValueByte < NumBytes; ++ValueByte)
+		{
+			Bytes[OutByteIndex] = (Value >> 8 * ValueByte) & 0xFF;
+			OutByteIndex += ByteIndexStep;
+		}
+
+		return Bytes;
+	}
 
 	/** Converts a normalized value to a byte array. Normalized value has to be in the 0-1 range. Assumes max signal format is 24bit. */
-	static TArray<uint8> NormalizedDMXValueToByteArray(float NormalizedValue, EDMXFixtureSignalFormat SignalFormat, bool bLSBOrder);
+	template <typename T, typename U = typename TEnableIf<TIsSame<T, float>::Value>::Type>
+	static TArray<uint8> NormalizedDMXValueToByteArray(T NormalizedValue, EDMXFixtureSignalFormat SignalFormat, bool bLSBOrder)
+	{
+		NormalizedValue = FMath::Clamp(NormalizedValue, 0.f, 1.f);
+		uint32 Value = FMath::Floor<uint32>(TNumericLimits<uint32>::Max() * static_cast<double>(NormalizedValue));
+
+		// Shift the value into signal format range
+		Value = Value >> ((3 - static_cast<uint8>(SignalFormat)) * 8);
+
+		return UnsignedInt32ToByteArray(Value, SignalFormat, bLSBOrder);
+	}
+
+	/** Returns the number of Bytes the Signal Format uses */
+	static uint8 GetSizeOfSignalFormat(EDMXFixtureSignalFormat SignalFormat);
+
+	/**
+	 * Returns the Max Value the Data Type can take
+	 *
+	 * @param DataType		The Signal Format in quesion.
+	 * @return				The Max Value of the Data Type, e.g. 255 for EDMXFixtureSignalFormat::E8Bit.
+	 */
+	static uint32 GetSignalFormatMaxValue(EDMXFixtureSignalFormat SignalFormat);
+
+	/** Clamps the Value to be in value range of the Signal Format */
+	static uint32 ClampValueBySignalFormat(uint32 Value, EDMXFixtureSignalFormat SignalFormat);
+	template <typename T>
+	static uint32 ClampValueBySignalFormat(T Value, EDMXFixtureSignalFormat SignalFormat) = delete;  // No implicit conversions
 
 private:
 	FDMXConversions() = delete;
diff --git a/Engine/Plugins/VirtualProduction/RemoteControlProtocolDMX/Source/RemoteControlProtocolDMX/Private/RemoteControlProtocolDMX.cpp b/Engine/Plugins/VirtualProduction/RemoteControlProtocolDMX/Source/RemoteControlProtocolDMX/Private/RemoteControlProtocolDMX.cpp
--- a/Engine/Plugins/VirtualProduction/RemoteControlProtocolDMX/Source/RemoteControlProtocolDMX/Private/RemoteControlProtocolDMX.cpp
+++ b/Engine/Plugins/VirtualProduction/RemoteControlProtocolDMX/Source/RemoteControlProtocolDMX/Private/RemoteControlProtocolDMX.cpp
@@ -2,6 +2,7 @@
 
 #include "RemoteControlProtocolDMX.h"
 
+#include "DMXConversions.h"
 #include "DMXProtocolSettings.h"
 #include "DMXProtocolTypes.h"
 #include "RemoteControlLogger.h"
@@ -222,7 +223,7 @@ void FRemoteControlProtocolDMX::ProcessAndApplyProtocolValue(const FDMXSignalSha
 	
 	FRemoteControlDMXProtocolEntity* DMXProtocolEntity = InProtocolEntityPtr->CastChecked<FRemoteControlDMXProtocolEntity>();
 	const uint8* ChannelData = &InSignal->ChannelData[InDMXOffset];	
-	const uint8 NumChannelsToOccupy = UDMXEntityFixtureType::NumChannelsToOccupy(DMXProtocolEntity->DataType);
+	const uint8 NumChannelsToOccupy = FDMXConversions::GetSizeOfSignalFormat(DMXProtocolEntity->DataType);
 
 	if(DMXProtocolEntity->CacheDMXBuffer.Num() != NumChannelsToOccupy ||
 		FMemory::Memcmp(DMXProtocolEntity->CacheDMXBuffer.GetData(), ChannelData, NumChannelsToOccupy) != 0)
diff --git a/Engine/Plugins/VirtualProduction/RemoteControlProtocolDMX/Source/RemoteControlProtocolDMXEditor/Private/RemoteControlProtocolDMXEditorCustomization.cpp b/Engine/Plugins/VirtualProduction/RemoteControlProtocolDMX/Source/RemoteControlProtocolDMXEditor/Private/RemoteControlProtocolDMXEditorCustomization.cpp
--- a/Engine/Plugins/VirtualProduction/RemoteControlProtocolDMX/Source/RemoteControlProtocolDMXEditor/Private/RemoteControlProtocolDMXEditorCustomization.cpp
+++ b/Engine/Plugins/VirtualProduction/RemoteControlProtocolDMX/Source/RemoteControlProtocolDMXEditor/Private/RemoteControlProtocolDMXEditorCustomization.cpp
@@ -2,6 +2,7 @@
 
 #include "RemoteControlProtocolDMXEditorCustomization.h"
 
+#include "DMXConversions.h"
 #include "DetailWidgetRow.h"
 #include "DMXProtocolTypes.h"
 #include "IDetailChildrenBuilder.h"
@@ -73,7 +74,7 @@ void FRemoteControlProtocolDMXEditorTypeCustomization::CheckAndApplyStartingChan
 		return;
 	}
 	const EDMXFixtureSignalFormat DataType = static_cast<EDMXFixtureSignalFormat>(DataTypeByte);
-	const uint8 NumChannelsToOccupy = UDMXEntityFixtureType::NumChannelsToOccupy(DataType);
+	const uint8 NumChannelsToOccupy = FDMXConversions::GetSizeOfSignalFormat(DataType);
 
 	// Get starting channel value
 	int32 StartingChannel = 0;
diff --git a/Engine/Plugins/VirtualProduction/RemoteControlProtocolDMX/Source/RemoteControlProtocolDMXEditor/RemoteControlProtocolDMXEditor.Build.cs b/Engine/Plugins/VirtualProduction/RemoteControlProtocolDMX/Source/RemoteControlProtocolDMXEditor/RemoteControlProtocolDMXEditor.Build.cs
--- a/Engine/Plugins/VirtualProduction/RemoteControlProtocolDMX/Source/RemoteControlProtocolDMXEditor/RemoteControlProtocolDMXEditor.Build.cs
+++ b/Engine/Plugins/VirtualProduction/RemoteControlProtocolDMX/Source/RemoteControlProtocolDMXEditor/RemoteControlProtocolDMXEditor.Build.cs
@@ -14,7 +14,7 @@ public class RemoteControlProtocolDMXEditor : ModuleRules
 			new string[] {
 				"Core",
 				"CoreUObject",
-				"DMXRuntime",
+				"DMXProtocol",
 				"InputCore",
 				"PropertyEditor",
 				"RemoteControl",
