DMX - Implement pixel mapping component multi selection drag drop, collapse matrix components by default.

diff --git a/Engine/Plugins/Experimental/DigitalMultiplex/DMXPixelMapping/Source/DMXPixelMappingEditor/Private/DMXPixelMappingComponentReference.h b/Engine/Plugins/Experimental/DigitalMultiplex/DMXPixelMapping/Source/DMXPixelMappingEditor/Private/DMXPixelMappingComponentReference.h
--- a/Engine/Plugins/Experimental/DigitalMultiplex/DMXPixelMapping/Source/DMXPixelMappingEditor/Private/DMXPixelMappingComponentReference.h
+++ b/Engine/Plugins/Experimental/DigitalMultiplex/DMXPixelMapping/Source/DMXPixelMappingEditor/Private/DMXPixelMappingComponentReference.h
@@ -45,14 +45,14 @@ public:
 		, Component(InComponent)
 	{}
 
+	friend FORCEINLINE uint32 GetTypeHash(const class FDMXPixelMappingComponentReference& ComponentRef)
+	{
+		return GetTypeHash(ComponentRef.GetComponent());
+	}
+
 private:
 	TWeakPtr<FDMXPixelMappingToolkit> ToolkitWeakPtr;
 
 	TWeakObjectPtr<UDMXPixelMappingBaseComponent> Component;
 };
 
-
-inline uint32 GetTypeHash(const class FDMXPixelMappingComponentReference& ComponentRef)
-{
-	return ::GetTypeHash(reinterpret_cast<void*>(ComponentRef.GetComponent()));
-}
diff --git a/Engine/Plugins/Experimental/DigitalMultiplex/DMXPixelMapping/Source/DMXPixelMappingEditor/Private/DragDrop/DMXPixelMappingDragDropOp.cpp b/Engine/Plugins/Experimental/DigitalMultiplex/DMXPixelMapping/Source/DMXPixelMappingEditor/Private/DragDrop/DMXPixelMappingDragDropOp.cpp
--- a/Engine/Plugins/Experimental/DigitalMultiplex/DMXPixelMapping/Source/DMXPixelMappingEditor/Private/DragDrop/DMXPixelMappingDragDropOp.cpp
+++ b/Engine/Plugins/Experimental/DigitalMultiplex/DMXPixelMapping/Source/DMXPixelMappingEditor/Private/DragDrop/DMXPixelMappingDragDropOp.cpp
@@ -17,120 +17,26 @@ TSharedRef<FDMXPixelMappingDragDropOp> FDMXPixelMappingDragDropOp::New(const TSh
 	Operation->Parent = InParent;
 	Operation->DefaultHoverText = InTemplate->Name;
 	Operation->CurrentHoverText = InTemplate->Name;
-	Operation->DragOffset = FVector2D::ZeroVector;
 
 	Operation->Construct();
 
 	return Operation;
 }
 
-TSharedRef<FDMXPixelMappingDragDropOp> FDMXPixelMappingDragDropOp::New(const TSet<FDMXPixelMappingComponentReference>& InComponentReferences)
+TSharedRef<FDMXPixelMappingDragDropOp> FDMXPixelMappingDragDropOp::New(const TArray<FDMXPixelMappingDragDropOp::FDraggingComponentReference>& InDraggingComponentReferences)
 {
 	TSharedRef<FDMXPixelMappingDragDropOp> Operation = MakeShared<FDMXPixelMappingDragDropOp>();
 
-	Operation->ComponentReferences = InComponentReferences;
+	Operation->DraggingComponentReferences = InDraggingComponentReferences;
 
-	for (const FDMXPixelMappingComponentReference& ComponentReference : InComponentReferences)
+	for (const FDMXPixelMappingDragDropOp::FDraggingComponentReference& DraggingComponentReference : InDraggingComponentReferences)
 	{
-		Operation->DefaultHoverText = FText::FromString(ComponentReference.GetComponent()->GetName());
-		Operation->CurrentHoverText = FText::FromString(ComponentReference.GetComponent()->GetName());
+		Operation->DefaultHoverText = FText::FromString(DraggingComponentReference.ComponentReference.GetComponent()->GetName());
+		Operation->CurrentHoverText = FText::FromString(DraggingComponentReference.ComponentReference.GetComponent()->GetName());
 		break;
 	}
-	Operation->DragOffset = FVector2D::ZeroVector;
 
 	Operation->Construct();
 
 	return Operation;
 }
-
-void FDMXPixelMappingDragDropOp::UpdateDragOffset(const FVector2D& DragStartScreenspacePosition)
-{
-	FArrangedWidget ArrangedWidget = GetArrangedWidgetFromComponent();
-	DragOffset = ArrangedWidget.Geometry.AbsoluteToLocal(DragStartScreenspacePosition);
-}
-
-void FDMXPixelMappingDragDropOp::SetComponentReferences(const TSet<FDMXPixelMappingComponentReference>& InComponentReferences)
-{
-	ComponentReferences = InComponentReferences;
-	for (const FDMXPixelMappingComponentReference& ComponentReference : InComponentReferences)
-	{
-		DefaultHoverText = FText::FromString(ComponentReference.GetComponent()->GetName());
-		CurrentHoverText = FText::FromString(ComponentReference.GetComponent()->GetName());
-		break;
-	}
-}
-
-FVector2D FDMXPixelMappingDragDropOp::GetDragOffset() const
-{
-	return DragOffset;
-}
-
-UDMXPixelMappingOutputComponent* FDMXPixelMappingDragDropOp::TryGetOutputComponent() const
-{
-	for (const FDMXPixelMappingComponentReference& ComponentReference : ComponentReferences)
-	{
-		return Cast<UDMXPixelMappingOutputComponent>(ComponentReference.GetComponent());
-	}
-
-	return nullptr;
-}
-
-UDMXPixelMappingBaseComponent* FDMXPixelMappingDragDropOp::TryGetBaseComponent() const
-{
-	for (const FDMXPixelMappingComponentReference& ComponentReference : ComponentReferences)
-	{
-		return Cast<UDMXPixelMappingBaseComponent>(ComponentReference.GetComponent());
-	}
-
-	return nullptr;
-}
-
-FArrangedWidget FDMXPixelMappingDragDropOp::GetArrangedWidgetFromComponent() const
-{
-	TSharedPtr<SWidget> WidgetToArrange;
-
-	for (const FDMXPixelMappingComponentReference& ComponentReference : ComponentReferences)
-	{
-		if (UDMXPixelMappingOutputComponent* OutputComponent = Cast<UDMXPixelMappingOutputComponent>(ComponentReference.GetComponent()))
-		{
-			// Use the parent component for group item and pixel components if they're locked in designer
-			if (OutputComponent->bLockInDesigner)
-			{
-				if (OutputComponent->GetClass() == UDMXPixelMappingFixtureGroupItemComponent::StaticClass() ||
-					OutputComponent->GetClass() == UDMXPixelMappingMatrixCellComponent::StaticClass())
-				{
-					UDMXPixelMappingOutputComponent* ParentOutputComponent = CastChecked<UDMXPixelMappingOutputComponent>(OutputComponent->Parent);
-					WidgetToArrange = ParentOutputComponent->GetCachedWidget();
-				}
-			}
-
-			WidgetToArrange = OutputComponent->GetCachedWidget();
-			break;
-		}
-	}
-
-	FArrangedWidget ArrangedWidget(SNullWidget::NullWidget, FGeometry());
-	GetArrangedWidget(WidgetToArrange.ToSharedRef(), ArrangedWidget);
-
-	return ArrangedWidget;
-}
-
-bool FDMXPixelMappingDragDropOp::GetArrangedWidget(TSharedRef<SWidget> Widget, FArrangedWidget& ArrangedWidget) const
-{
-	TSharedPtr<SWindow> WidgetWindow = FSlateApplication::Get().FindWidgetWindow(Widget);
-	if (!WidgetWindow.IsValid())
-	{
-		return false;
-	}
-
-	TSharedRef<SWindow> CurrentWindowRef = WidgetWindow.ToSharedRef();
-
-	FWidgetPath WidgetPath;
-	if (FSlateApplication::Get().GeneratePathToWidgetUnchecked(Widget, WidgetPath))
-	{
-		ArrangedWidget = WidgetPath.FindArrangedWidget(Widget).Get(FArrangedWidget::GetNullWidget());
-		return true;
-	}
-
-	return false;
-}
\ No newline at end of file
diff --git a/Engine/Plugins/Experimental/DigitalMultiplex/DMXPixelMapping/Source/DMXPixelMappingEditor/Private/DragDrop/DMXPixelMappingDragDropOp.h b/Engine/Plugins/Experimental/DigitalMultiplex/DMXPixelMapping/Source/DMXPixelMappingEditor/Private/DragDrop/DMXPixelMappingDragDropOp.h
--- a/Engine/Plugins/Experimental/DigitalMultiplex/DMXPixelMapping/Source/DMXPixelMappingEditor/Private/DragDrop/DMXPixelMappingDragDropOp.h
+++ b/Engine/Plugins/Experimental/DigitalMultiplex/DMXPixelMapping/Source/DMXPixelMappingEditor/Private/DragDrop/DMXPixelMappingDragDropOp.h
@@ -23,29 +23,22 @@ class FDMXPixelMappingDragDropOp
 public:
 	DRAG_DROP_OPERATOR_TYPE(FDMXPixelMappingDragDropOp, FDecoratedDragDropOp)
 
-	/** Constructs the drag drop operation */
-	static TSharedRef<FDMXPixelMappingDragDropOp> New(const TSharedPtr<FDMXPixelMappingComponentTemplate>& InTemplate, UDMXPixelMappingBaseComponent* InParent = nullptr);
-	static TSharedRef<FDMXPixelMappingDragDropOp> New(const TSet<FDMXPixelMappingComponentReference>& InComponentReferences);
-
-	void UpdateDragOffset(const FVector2D& DragStartScreenspacePosition);
-
-	void SetComponentReferences(const TSet<FDMXPixelMappingComponentReference>& InComponentReferences);
+	struct FDraggingComponentReference
+	{
+		FDMXPixelMappingComponentReference ComponentReference;
 
-	/** Gets the offset of the drag drop op from the mouse pos (context menu location) */
-	FVector2D GetDragOffset() const;
+		FVector2D DraggedOffset = FVector2D::ZeroVector;
+	};
 
-	/** Returns an output component or nullptr if it's not an output component drag drop op */
-	UDMXPixelMappingOutputComponent* TryGetOutputComponent() const;
-
-	/** Returns an base component or nullptr if it's not an output component drag drop op */
-	UDMXPixelMappingBaseComponent* TryGetBaseComponent() const;
+	/** Constructs the drag drop operation */
+	static TSharedRef<FDMXPixelMappingDragDropOp> New(const TSharedPtr<FDMXPixelMappingComponentTemplate>& InTemplate, UDMXPixelMappingBaseComponent* InParent = nullptr);
+	static TSharedRef<FDMXPixelMappingDragDropOp> New(const TArray<FDMXPixelMappingDragDropOp::FDraggingComponentReference>& InComponentReferences);
 
-protected:
-	/** Gets an arranged widget from the dragged component */
-	FArrangedWidget GetArrangedWidgetFromComponent() const;
+	/** Returns the dragged component references */
+	const TArray<FDMXPixelMappingDragDropOp::FDraggingComponentReference>& GetDraggingComponentReferences() const { return DraggingComponentReferences; }
 
-	/** Gets an arranged widget from a widget */
-	bool GetArrangedWidget(TSharedRef<SWidget> Widget, FArrangedWidget& ArrangedWidget) const;
+	/** Sets the dragged component references */
+	void SetDraggingComponentReferences(const TArray<FDMXPixelMappingDragDropOp::FDraggingComponentReference>& InDraggingComponentReferences) { DraggingComponentReferences = InDraggingComponentReferences; }
 
 public:
 	/** The template to create an instance */
@@ -53,8 +46,5 @@ public:
 
 	TWeakObjectPtr<UDMXPixelMappingBaseComponent> Parent;
 
-	TSet<FDMXPixelMappingComponentReference> ComponentReferences;
-
-private:
-	FVector2D DragOffset;
+	TArray<FDMXPixelMappingDragDropOp::FDraggingComponentReference> DraggingComponentReferences;
 };
diff --git a/Engine/Plugins/Experimental/DigitalMultiplex/DMXPixelMapping/Source/DMXPixelMappingEditor/Private/Views/SDMXPixelMappingDesignerView.cpp b/Engine/Plugins/Experimental/DigitalMultiplex/DMXPixelMapping/Source/DMXPixelMappingEditor/Private/Views/SDMXPixelMappingDesignerView.cpp
--- a/Engine/Plugins/Experimental/DigitalMultiplex/DMXPixelMapping/Source/DMXPixelMappingEditor/Private/Views/SDMXPixelMappingDesignerView.cpp
+++ b/Engine/Plugins/Experimental/DigitalMultiplex/DMXPixelMapping/Source/DMXPixelMappingEditor/Private/Views/SDMXPixelMappingDesignerView.cpp
@@ -34,10 +34,31 @@
 
 #define LOCTEXT_NAMESPACE "SDMXPixelMappingDesignerView"
 
-void SDMXPixelMappingDesignerView::Construct(const FArguments& InArgs, const TSharedPtr<FDMXPixelMappingToolkit>& InToolkit)
+namespace
 {
-	bMovingExistingWidget = false;
+	bool GetArrangedWidget(TSharedRef<SWidget> Widget, FArrangedWidget& ArrangedWidget)
+	{
+		TSharedPtr<SWindow> WidgetWindow = FSlateApplication::Get().FindWidgetWindow(Widget);
+		if (!WidgetWindow.IsValid())
+		{
+			return false;
+		}
 
+		TSharedRef<SWindow> CurrentWindowRef = WidgetWindow.ToSharedRef();
+
+		FWidgetPath WidgetPath;
+		if (FSlateApplication::Get().GeneratePathToWidgetUnchecked(Widget, WidgetPath))
+		{
+			ArrangedWidget = WidgetPath.FindArrangedWidget(Widget).Get(FArrangedWidget::GetNullWidget());
+			return true;
+		}
+
+		return false;
+	}
+}
+
+void SDMXPixelMappingDesignerView::Construct(const FArguments& InArgs, const TSharedPtr<FDMXPixelMappingToolkit>& InToolkit)
+{
 	DelegateHandleChangeComponents = InToolkit->GetOnComponenetAddedOrDeletedDelegate().AddSP(this, &SDMXPixelMappingDesignerView::HandleChangeComponents);
 	OnSelectedComponenetChangedHandle = InToolkit->GetOnSelectedComponenetChangedDelegate().AddRaw(this, &SDMXPixelMappingDesignerView::OnSelectedComponenetChanged);
 
@@ -199,8 +220,6 @@ FReply SDMXPixelMappingDesignerView::OnMouseButtonDown(const FGeometry& MyGeomet
 		bFoundWidgetUnderCursor = FindComponentUnderCursor(MyGeometry, MouseEvent, UDMXPixelMappingOutputComponent::StaticClass(), HitResult);
 		if (bFoundWidgetUnderCursor)
 		{
-			DragOffset = HitResult.WidgetArranged.Geometry.AbsoluteToLocal(MouseEvent.GetScreenSpacePosition());
-			
 			if (UDMXPixelMappingMatrixCellComponent* MatrixPixelComponent = Cast<UDMXPixelMappingMatrixCellComponent>(HitResult.Component))
 			{
 				// If a matrix pixel component is selected and it is locked in designer, select the owning Matrix Component instead
@@ -215,10 +234,24 @@ FReply SDMXPixelMappingDesignerView::OnMouseButtonDown(const FGeometry& MyGeomet
 
 	if (MouseEvent.GetEffectingButton() == EKeys::LeftMouseButton)
 	{
-		// Select and detect drag when something was clicked
-		if (bFoundWidgetUnderCursor && MouseEvent.GetEffectingButton() == EKeys::LeftMouseButton)
+		// Detect drag when something was clicked
+		if (bFoundWidgetUnderCursor)
 		{
-			ResolvePendingSelectedComponents(MouseEvent);
+			const bool bClickedSelectedComponent = [this]()
+			{
+				for (const FDMXPixelMappingComponentReference& ComponentReference : GetSelectedComponents())
+				{
+					if (ComponentReference.GetComponent() == PendingSelectedComponent)
+					{
+						return true;
+					}
+				}
+
+				return false;
+			}();
+
+			bool bClearPreviousSelection = !MouseEvent.IsShiftDown() && !MouseEvent.IsControlDown() && !bClickedSelectedComponent;
+			ResolvePendingSelectedComponents(bClearPreviousSelection);
 
 			DraggingStartPositionScreenSpace = MouseEvent.GetScreenSpacePosition();
 
@@ -230,7 +263,6 @@ FReply SDMXPixelMappingDesignerView::OnMouseButtonDown(const FGeometry& MyGeomet
 				.DetectDrag(AsShared(), EKeys::LeftMouseButton);
 		}
 
-
 		// We didn't click anything. Clear the selection.
 		PendingSelectedComponent = nullptr;
 
@@ -256,14 +288,12 @@ FReply SDMXPixelMappingDesignerView::OnMouseButtonDown(const FGeometry& MyGeomet
 
 FReply SDMXPixelMappingDesignerView::OnMouseButtonUp(const FGeometry& MyGeometry, const FPointerEvent& MouseEvent)
 {
-	if (HasMouseCapture() && MouseEvent.GetEffectingButton() == EKeys::LeftMouseButton)
-	{
-		bMovingExistingWidget = false;
-	}
-	
+	bDraggingNewGroupItems = false;
+
 	SDMXPixelMappingSurface::OnMouseButtonUp(MyGeometry, MouseEvent);
 
-	ResolvePendingSelectedComponents(MouseEvent);
+	constexpr bool bClearPreviousSelection = true;
+	ResolvePendingSelectedComponents(bClearPreviousSelection);
 
 	return FReply::Handled().ReleaseMouseCapture();
 }
@@ -275,8 +305,6 @@ FReply SDMXPixelMappingDesignerView::OnMouseMove(const FGeometry& MyGeometry, co
 		return FReply::Unhandled();
 	}
 
-	CachedMousePosition = MouseEvent.GetScreenSpacePosition();
-
 	FReply SurfaceHandled = SDMXPixelMappingSurface::OnMouseMove(MyGeometry, MouseEvent);
 	if (SurfaceHandled.IsEventHandled())
 	{
@@ -299,7 +327,6 @@ FReply SDMXPixelMappingDesignerView::OnMouseMove(const FGeometry& MyGeometry, co
 
 		if (!bIsRootWidgetSelected)
 		{
-			bMovingExistingWidget = true;
 			//Drag selected widgets
 			return FReply::Handled();
 		}
@@ -382,9 +409,49 @@ FReply SDMXPixelMappingDesignerView::OnDragDetected(const FGeometry& MyGeometry,
 		// Clear any pending selected widgets, the user has already decided what widget they want.
 		PendingSelectedComponent = nullptr;
 		
-		TSharedRef<FDMXPixelMappingDragDropOp> DragDropOp = FDMXPixelMappingDragDropOp::New(SelectedComponents);
+		TArray<FDMXPixelMappingDragDropOp::FDraggingComponentReference> DraggingComponentCandidates;
+
+		for (const FDMXPixelMappingComponentReference& SelectedComponent : GetSelectedComponents())
+		{
+			if (UDMXPixelMappingOutputComponent* OutputComponent = Cast<UDMXPixelMappingOutputComponent>(SelectedComponent.GetComponent()))
+			{
+				// Determine The offset to keep the widget from the mouse while dragging
+				FArrangedWidget ArrangedWidget(SNullWidget::NullWidget, FGeometry());
+				GetArrangedWidget(OutputComponent->GetCachedWidget().ToSharedRef(), ArrangedWidget);
+
+				FDMXPixelMappingDragDropOp::FDraggingComponentReference DraggingComponentReference;
+				DraggingComponentReference.ComponentReference = SelectedComponent;
+				DraggingComponentReference.DraggedOffset = ArrangedWidget.Geometry.AbsoluteToLocal(DraggingStartPositionScreenSpace) / ArrangedWidget.Geometry.GetLocalSize();
+				DraggingComponentCandidates.Add(DraggingComponentReference);
+			}
+		}
 
-		DragDropOp->UpdateDragOffset(DraggingStartPositionScreenSpace);
+		TArray<FDMXPixelMappingDragDropOp::FDraggingComponentReference> DraggingComponents;
+
+		// Check the parent chain of each dragged component and ignore those that are children of other dragged components
+		for (const FDMXPixelMappingDragDropOp::FDraggingComponentReference& Candidate : DraggingComponentCandidates)
+		{
+			bool bIsChild = false;
+			for (UDMXPixelMappingBaseComponent* ParentPtr = Candidate.ComponentReference.GetComponent()->Parent; ParentPtr != nullptr; ParentPtr = ParentPtr->Parent)
+			{
+				bIsChild = DraggingComponentCandidates.ContainsByPredicate([ParentPtr](const FDMXPixelMappingDragDropOp::FDraggingComponentReference& Ref)
+					{ 
+						return Ref.ComponentReference.GetComponent() == ParentPtr; 
+					});
+
+				if (bIsChild)
+				{
+					break;
+				}
+			}
+
+			if (!bIsChild)
+			{
+				DraggingComponents.Add(Candidate);
+			}
+		}
+
+		TSharedRef<FDMXPixelMappingDragDropOp> DragDropOp = FDMXPixelMappingDragDropOp::New(DraggingComponents);
 		DragDropOp->SetDecoratorVisibility(false);
 
 		return FReply::Handled().BeginDragDrop(DragDropOp);
@@ -421,19 +488,88 @@ FReply SDMXPixelMappingDesignerView::OnDragOver(const FGeometry& MyGeometry, con
 
 	if (DragDropEvent.GetOperation().IsValid() && DragDropEvent.GetOperation()->IsOfType<FDMXPixelMappingDragDropOp>())
 	{
-		TSharedPtr<FDMXPixelMappingDragDropOp> SelectedDragDropOp = DragDropEvent.GetOperationAs<FDMXPixelMappingDragDropOp>();
+		const FGeometry WidgetUnderCursorGeometry = PreviewSizeConstraint->GetTickSpaceGeometry();
+		const FVector2D ScreenSpacePosition = DragDropEvent.GetScreenSpacePosition();
+		const FVector2D LocalPosition = WidgetUnderCursorGeometry.AbsoluteToLocal(ScreenSpacePosition);
 
-		if (UDMXPixelMappingOutputComponent* OutputComponent = SelectedDragDropOp->TryGetOutputComponent())
+		const TSharedPtr<FDMXPixelMappingDragDropOp> SelectedDragDropOp = DragDropEvent.GetOperationAs<FDMXPixelMappingDragDropOp>();
+		const TArray<FDMXPixelMappingDragDropOp::FDraggingComponentReference>& DraggingComponentReferences = SelectedDragDropOp->GetDraggingComponentReferences();
+
+		if (bDraggingNewGroupItems && DraggingComponentReferences.Num() > 0)
 		{
-			FGeometry WidgetUnderCursorGeometry = PreviewSizeConstraint->GetTickSpaceGeometry();
-			FVector2D ScreenSpacePosition = DragDropEvent.GetScreenSpacePosition();
-			FVector2D LocalPosition = WidgetUnderCursorGeometry.AbsoluteToLocal(ScreenSpacePosition);
+			// If new group items are dragged, lay them out over their group dynamically
+			int32 CountGroupItems = 0;
+
+			if (UDMXPixelMappingFixtureGroupComponent* GroupComponent = Cast<UDMXPixelMappingFixtureGroupComponent>(DraggingComponentReferences[0].ComponentReference.GetComponent()->Parent))
+			{
+				FArrangedWidget GroupArrangedWidget(SNullWidget::NullWidget, FGeometry());
+				GetArrangedWidget(GroupComponent->GetCachedWidget().ToSharedRef(), GroupArrangedWidget);
 
-			FVector2D Offset = SelectedDragDropOp->GetDragOffset();
+				const FVector2D GroupPos = GroupArrangedWidget.Geometry.GetLocalPositionAtCoordinates(FVector2D(0.f, 0.f));
+				const FVector2D GroupSize = GroupArrangedWidget.Geometry.GetLocalSize();
 
-			FVector2D NewPosition = LocalPosition - Offset;
+				const float StartPosX = LocalPosition.X > GroupPos.X && LocalPosition.X < GroupPos.X + GroupSize.X ? LocalPosition.X : GroupPos.X;
+				const float StartPosY = LocalPosition.Y > GroupPos.Y && LocalPosition.Y < GroupPos.Y + GroupSize.Y ? LocalPosition.Y : GroupPos.Y;
 
-			OutputComponent->SetPosition(NewPosition);
+				const float DesiredRectX = GroupSize.X - (StartPosX - GroupPos.X);
+				const float DesiredRectY = GroupSize.Y - (StartPosY - GroupPos.Y);
+
+				const FVector2D StartPos = FVector2D(StartPosX, StartPosY);
+
+				FVector2D CellSize = FVector2D::ZeroVector;
+				for (const FDMXPixelMappingDragDropOp::FDraggingComponentReference& DraggingComponentReference : SelectedDragDropOp->GetDraggingComponentReferences())
+				{
+					if (UDMXPixelMappingOutputComponent* OutputComponent = Cast<UDMXPixelMappingOutputComponent>(DraggingComponentReference.ComponentReference.GetComponent()))
+					{
+						CellSize = FMath::Max(CellSize, OutputComponent->GetSize());
+					}
+				}
+
+				const int32 NumTiles = SelectedDragDropOp->GetDraggingComponentReferences().Num();
+				const int32 NumCollums = DesiredRectX / CellSize.X;
+
+				int32 IndexCollumn = 0;
+				int32 IndexRow = 0;
+				for (int32 IndexComponent = 0; IndexComponent < NumTiles; IndexComponent++, IndexCollumn++)
+				{
+					if (UDMXPixelMappingOutputComponent* OutputComponent = Cast<UDMXPixelMappingOutputComponent>(SelectedDragDropOp->GetDraggingComponentReferences()[IndexComponent].ComponentReference.GetComponent()))
+					{
+						FVector2D NewPosition;
+						if (IndexRow * CellSize.Y > GroupSize.Y)
+						{
+							NewPosition = StartPos + FVector2D((NumCollums - 1) * CellSize.X, IndexRow * CellSize.Y);
+						}
+						else
+						{
+							if (IndexCollumn == NumCollums)
+							{
+								IndexRow++;
+								IndexCollumn = 0;
+							}
+
+							NewPosition = StartPos + FVector2D(IndexCollumn * CellSize.X, IndexRow * CellSize.Y);
+						}
+
+						OutputComponent->SetPosition(NewPosition);
+					}
+				}
+			}
+		}
+		else
+		{
+			for (const FDMXPixelMappingDragDropOp::FDraggingComponentReference& DraggingComponentReference : SelectedDragDropOp->GetDraggingComponentReferences())
+			{
+				if (UDMXPixelMappingOutputComponent* OutputComponent = Cast<UDMXPixelMappingOutputComponent>(DraggingComponentReference.ComponentReference.GetComponent()))
+				{
+					FArrangedWidget ArrangedWidget(SNullWidget::NullWidget, FGeometry());
+					GetArrangedWidget(OutputComponent->GetCachedWidget().ToSharedRef(), ArrangedWidget);
+
+					const FVector2D Offset = DraggingComponentReference.DraggedOffset * ArrangedWidget.Geometry.GetLocalSize();
+					const FVector2D NewPosition = LocalPosition - Offset;
+
+					OutputComponent->SetPosition(NewPosition);
+				}
+			}
 		}
 	}
 
@@ -444,13 +580,13 @@ FReply SDMXPixelMappingDesignerView::OnDrop(const FGeometry& MyGeometry, const F
 {
 	SDMXPixelMappingSurface::OnDrop(MyGeometry, DragDropEvent);
 
-	bMovingExistingWidget = false;
-
 	DropComponent(MyGeometry, DragDropEvent);
 
 	CachedRendererComponent = nullptr;
 	UpdateOutput(false);
 
+	bDraggingNewGroupItems = false;
+
 	return FReply::Handled().EndDragDrop();
 }
 
@@ -839,18 +975,18 @@ EVisibility SDMXPixelMappingDesignerView::GetTitleBarVisibility() const
 	return EVisibility::Collapsed;
 }
 
-void SDMXPixelMappingDesignerView::ResolvePendingSelectedComponents(const FPointerEvent& MouseEvent)
+void SDMXPixelMappingDesignerView::ResolvePendingSelectedComponents(bool bClearPreviousSelection)
 {
 	if (PendingSelectedComponent.IsValid())
 	{
 		TSet<FDMXPixelMappingComponentReference> SelectedComponents;
-		if (MouseEvent.IsShiftDown() || MouseEvent.IsControlDown())
+		if (!bClearPreviousSelection)
 		{
 			SelectedComponents = ToolkitWeakPtr.Pin()->GetSelectedComponents();
 		}
 		SelectedComponents.Add(ToolkitWeakPtr.Pin()->GetReferenceFromComponent(PendingSelectedComponent.Get()));
 		ToolkitWeakPtr.Pin()->SelectComponents(SelectedComponents);
-
+		
 		PendingSelectedComponent = nullptr;
 	}
 }
@@ -859,7 +995,7 @@ void SDMXPixelMappingDesignerView::AddComponentFromPalette(const FGeometry& MyGe
 {
 	// Ignore calls where the component already exists
 	if (!TemplateDragDropOp.IsValid() || 
-		TemplateDragDropOp->TryGetBaseComponent())
+		TemplateDragDropOp->GetDraggingComponentReferences().Num() > 0)
 	{
 		return;
 	}
@@ -876,10 +1012,12 @@ void SDMXPixelMappingDesignerView::AddComponentFromPalette(const FGeometry& MyGe
 			{
 				TSet<FDMXPixelMappingComponentReference> SelectedComponents;
 
+				// Special case for fixture group as we want to allow multiple patches to be created on the fly by dragging
+				// and we don't want to add the same patch twice
 				if (UDMXPixelMappingFixtureGroupComponent* FixtureGroupComponent = Cast<UDMXPixelMappingFixtureGroupComponent>(Target))
 				{
-					// Special case for fixture group as we want to allow multiple patches to be created on the fly by dragging
-					// and we don't want to add the same patch twice
+					bDraggingNewGroupItems = true;
+
 					for (const FDMXEntityFixturePatchRef& SelectedFixturePatchRef : FixtureGroupComponent->SelectedFixturePatchRef)
 					{
 						// Don't add the same patch twice
@@ -901,17 +1039,11 @@ void SDMXPixelMappingDesignerView::AddComponentFromPalette(const FGeometry& MyGe
 								Target->AddChild(Component);
 								Component->PostParentAssigned();
 								SelectedComponents.Add(ToolkitWeakPtr.Pin()->GetReferenceFromComponent(Component));
+
+								bDraggingNewGroupItems = true;
 							}
 						}
 					}
-
-					// If multi-drag drop terminate the drag drop op to not get akward behaviour
-					if (SelectedComponents.Num() > 1)
-					{
-						bRequestTerminateDragDrop = true;
-
-						SelectedComponents.Empty();
-					}
 				}
 				else if (UDMXPixelMappingMatrixComponent* MatrixComponent = Cast<UDMXPixelMappingMatrixComponent>(Target))
 				{
@@ -943,7 +1075,23 @@ void SDMXPixelMappingDesignerView::AddComponentFromPalette(const FGeometry& MyGe
 				CreateExtensionWidgetsForSelection();
 
 				// Update the drag drop op with the newly selected components
-				TemplateDragDropOp->SetComponentReferences(SelectedComponents);
+				TArray<FDMXPixelMappingDragDropOp::FDraggingComponentReference> DraggingComponentReferences;
+				for (const FDMXPixelMappingComponentReference& ComponentReference : SelectedComponents)
+				{
+					if (UDMXPixelMappingOutputComponent* OutputComponent = Cast<UDMXPixelMappingOutputComponent>(ComponentReference.GetComponent()))
+					{
+						// Determine The offset to keep the widget from the mouse while dragging
+						FArrangedWidget ArrangedWidget(SNullWidget::NullWidget, FGeometry());
+						GetArrangedWidget(OutputComponent->GetCachedWidget().ToSharedRef(), ArrangedWidget);
+
+						FDMXPixelMappingDragDropOp::FDraggingComponentReference DraggingComponentReference;
+						DraggingComponentReference.ComponentReference = ComponentReference;
+						DraggingComponentReference.DraggedOffset = FVector2D::ZeroVector; // ArrangedWidget.Geometry.AbsoluteToLocal(DraggingStartPositionScreenSpace) / ArrangedWidget.Geometry.GetLocalSize();
+						DraggingComponentReferences.Add(DraggingComponentReference);
+					}
+				}
+
+				TemplateDragDropOp->SetDraggingComponentReferences(DraggingComponentReferences);
 			}
 		}
 	}
@@ -954,30 +1102,28 @@ void SDMXPixelMappingDesignerView::DropComponent(const FGeometry& MyGeometry, co
 	TSharedPtr<FDMXPixelMappingDragDropOp> SelectedDragDropOp = DragDropEvent.GetOperationAs<FDMXPixelMappingDragDropOp>();
 	if (SelectedDragDropOp.IsValid())
 	{
-		if (UDMXPixelMappingOutputComponent* OutputComponent = SelectedDragDropOp->TryGetOutputComponent())
-		{
-			FGeometry WidgetUnderCursorGeometry = PreviewSizeConstraint->GetTickSpaceGeometry();
-			FVector2D ScreenSpacePosition = DragDropEvent.GetScreenSpacePosition();
-			FVector2D LocalPosition = WidgetUnderCursorGeometry.AbsoluteToLocal(ScreenSpacePosition);
-
-			FVector2D Offset = SelectedDragDropOp->GetDragOffset();
-			FVector2D NewPosition = LocalPosition - Offset;
+		FGeometry WidgetUnderCursorGeometry = PreviewSizeConstraint->GetTickSpaceGeometry();
+		FVector2D ScreenSpacePosition = DragDropEvent.GetScreenSpacePosition();
+		FVector2D LocalPosition = WidgetUnderCursorGeometry.AbsoluteToLocal(ScreenSpacePosition);
 
-			OutputComponent->SetPosition(NewPosition);
-
-			// Assign a new ZOrder to dropped components
-			SDMXPixelMappingDesignerView::FComponentHitResult HitResult;
-			FindComponentUnderCursor(MyGeometry, DragDropEvent, UDMXPixelMappingOutputComponent::StaticClass(), HitResult);
-
-			if (UDMXPixelMappingOutputComponent* TargetOutputComponent = Cast<UDMXPixelMappingOutputComponent>(HitResult.Component))
+		for (const FDMXPixelMappingDragDropOp::FDraggingComponentReference& DraggingComponentReference : SelectedDragDropOp->GetDraggingComponentReferences())
+		{
+			if (UDMXPixelMappingOutputComponent* OutputComponent = Cast<UDMXPixelMappingOutputComponent>(DraggingComponentReference.ComponentReference.GetComponent()))
 			{
-				if (OutputComponent != TargetOutputComponent &&
-					OutputComponent->GetZOrder() <= TargetOutputComponent->GetZOrder())
+				// Assign a new ZOrder to dropped components
+				SDMXPixelMappingDesignerView::FComponentHitResult HitResult;
+				FindComponentUnderCursor(MyGeometry, DragDropEvent, UDMXPixelMappingOutputComponent::StaticClass(), HitResult);
+
+				if (UDMXPixelMappingOutputComponent* TargetOutputComponent = Cast<UDMXPixelMappingOutputComponent>(HitResult.Component))
 				{
-					int32 NewZOrder = TargetOutputComponent->GetZOrder() + 1;
+					if (OutputComponent != TargetOutputComponent &&
+						OutputComponent->GetZOrder() <= TargetOutputComponent->GetZOrder())
+					{
+						int32 NewZOrder = TargetOutputComponent->GetZOrder() + 1;
 
-					OutputComponent->Modify();
-					OutputComponent->SetZOrder(NewZOrder);
+						OutputComponent->Modify();
+						OutputComponent->SetZOrder(NewZOrder);
+					}
 				}
 			}
 		}
diff --git a/Engine/Plugins/Experimental/DigitalMultiplex/DMXPixelMapping/Source/DMXPixelMappingEditor/Private/Views/SDMXPixelMappingDesignerView.h b/Engine/Plugins/Experimental/DigitalMultiplex/DMXPixelMapping/Source/DMXPixelMappingEditor/Private/Views/SDMXPixelMappingDesignerView.h
--- a/Engine/Plugins/Experimental/DigitalMultiplex/DMXPixelMapping/Source/DMXPixelMappingEditor/Private/Views/SDMXPixelMappingDesignerView.h
+++ b/Engine/Plugins/Experimental/DigitalMultiplex/DMXPixelMapping/Source/DMXPixelMappingEditor/Private/Views/SDMXPixelMappingDesignerView.h
@@ -121,7 +121,7 @@ private:
 	void OnSelectedComponenetChanged();
 
 	/** Adds any pending selected Components to the selection set */
-	void ResolvePendingSelectedComponents(const FPointerEvent& MouseEvent);
+	void ResolvePendingSelectedComponents(bool bClearPreviousSelection = false);
 	
 	/** Adds a new component from a drag drop op from the palette */
 	void AddComponentFromPalette(const FGeometry& MyGeometry, const TSharedPtr<FDMXPixelMappingDragDropOp>& TemplateDragDropOp);
@@ -176,19 +176,13 @@ private:
 	/** The ruler bar on the left side of the designer. */
 	TSharedPtr<SDMXPixelMappingRuler> SideRuler;
 
-	/** Cache last mouse position to be used as a paste drop location */
-	FVector2D CachedMousePosition;
-
-	/** The location in selected widget local space where the context menu was summoned. */
-	FVector2D DragOffset;
-
 	TWeakObjectPtr<UDMXPixelMappingBaseComponent> PendingSelectedComponent;
 
-	/** True if an existing widget is being moved in its current container, or into a new container. */
-	bool bMovingExistingWidget;
+	/** True if new components are being dragged. */
+	bool bDraggingNewGroupItems = false;
 
 	/** If true, terminates any existing drag drop op without handling it */
-	bool bRequestTerminateDragDrop;
+	bool bRequestTerminateDragDrop = false;
 
 	/** The position in screen space where the user began dragging a widget */
 	FVector2D DraggingStartPositionScreenSpace;
diff --git a/Engine/Plugins/Experimental/DigitalMultiplex/DMXPixelMapping/Source/DMXPixelMappingEditor/Private/Views/SDMXPixelMappingHierarchyView.cpp b/Engine/Plugins/Experimental/DigitalMultiplex/DMXPixelMapping/Source/DMXPixelMappingEditor/Private/Views/SDMXPixelMappingHierarchyView.cpp
--- a/Engine/Plugins/Experimental/DigitalMultiplex/DMXPixelMapping/Source/DMXPixelMappingEditor/Private/Views/SDMXPixelMappingHierarchyView.cpp
+++ b/Engine/Plugins/Experimental/DigitalMultiplex/DMXPixelMapping/Source/DMXPixelMappingEditor/Private/Views/SDMXPixelMappingHierarchyView.cpp
@@ -1,11 +1,13 @@
 // Copyright Epic Games, Inc. All Rights Reserved.
 
 #include "Views/SDMXPixelMappingHierarchyView.h"
-#include "Widgets/SDMXPixelMappingHierarchyItem.h"
-#include "Toolkits/DMXPixelMappingToolkit.h"
+
 #include "DMXPixelMappingComponentReference.h"
 #include "DMXPixelMappingEditorUtils.h"
+#include "Components/DMXPixelMappingMatrixComponent.h"
+#include "Toolkits/DMXPixelMappingToolkit.h"
 #include "ViewModels/DMXPixelMappingHierarchyViewModel.h"
+#include "Widgets/SDMXPixelMappingHierarchyItem.h"
 
 #include "Framework/Commands/GenericCommands.h"
 #include "Widgets/Layout/SBorder.h"
@@ -107,16 +109,16 @@ void SDMXPixelMappingHierarchyView::Construct(const FArguments& InArgs, const TS
 		[
 			SNew(SBorder)
 			.BorderImage(FEditorStyle::GetBrush("ToolPanel.GroupBorder"))
-		[
-			SNew(SVerticalBox)
-			+ SVerticalBox::Slot()
-		.FillHeight(1.0f)
-		[
-			SAssignNew(TreeViewArea, SBorder)
-			.Padding(0)
-		.BorderImage(FEditorStyle::GetBrush("NoBrush"))
-		]
-		]
+			[
+				SNew(SVerticalBox)
+				+ SVerticalBox::Slot()
+				.FillHeight(1.0f)
+				[
+					SAssignNew(TreeViewArea, SBorder)
+					.Padding(0)
+					.BorderImage(FEditorStyle::GetBrush("NoBrush"))
+				]
+			]
 		];
 
 	RebuildTreeView();
@@ -268,7 +270,15 @@ void SDMXPixelMappingHierarchyView::ConditionallyUpdateTree()
 			bool* pPreviousExpansionState = ComponentExpansionStates.Find(Component);
 			if(pPreviousExpansionState == nullptr)
 			{
-				TreeView->SetItemExpansion(Model, true);
+				// Initially collapse matrix components
+				if (Cast<UDMXPixelMappingMatrixComponent>(Component))
+				{
+					TreeView->SetItemExpansion(Model, false);
+				}
+				else
+				{
+					TreeView->SetItemExpansion(Model, true);
+				}
 			}
 			else
 			{
diff --git a/Engine/Plugins/Experimental/DigitalMultiplex/DMXPixelMapping/Source/DMXPixelMappingEditor/Private/Widgets/SDMXPixelMappingHierarchyItem.cpp b/Engine/Plugins/Experimental/DigitalMultiplex/DMXPixelMapping/Source/DMXPixelMappingEditor/Private/Widgets/SDMXPixelMappingHierarchyItem.cpp
--- a/Engine/Plugins/Experimental/DigitalMultiplex/DMXPixelMapping/Source/DMXPixelMappingEditor/Private/Widgets/SDMXPixelMappingHierarchyItem.cpp
+++ b/Engine/Plugins/Experimental/DigitalMultiplex/DMXPixelMapping/Source/DMXPixelMappingEditor/Private/Widgets/SDMXPixelMappingHierarchyItem.cpp
@@ -75,44 +75,52 @@ FReply SDMXPixelMappingHierarchyItem::OnDraggingWidget(const FGeometry& MyGeomet
 {
 	check(Model.IsValid());
 
-	return FReply::Handled().BeginDragDrop(FDMXPixelMappingDragDropOp::New({ Model.Pin()->GetReference() }));
+	TArray<FDMXPixelMappingDragDropOp::FDraggingComponentReference> DraggingComponents;
+	FDMXPixelMappingDragDropOp::FDraggingComponentReference DraggingComponentReference;
+	DraggingComponentReference.ComponentReference = Model.Pin()->GetReference();
+	DraggingComponentReference.DraggedOffset = FVector2D::ZeroVector;
+	DraggingComponents.Add(DraggingComponentReference);
+
+	return FReply::Handled().BeginDragDrop(FDMXPixelMappingDragDropOp::New(DraggingComponents));
 }
 
 FReply SDMXPixelMappingHierarchyItem::OnDropWidget(const FDragDropEvent& InDragDropEvent)
 {
 	const FScopedTransaction Transaction(FText::FromString("MoveComponent"));
 	TSharedPtr<FDMXPixelMappingDragDropOp> ComponentDragDropOp = InDragDropEvent.GetOperationAs<FDMXPixelMappingDragDropOp>();
-	if (ComponentDragDropOp.IsValid())
+	if (ComponentDragDropOp.IsValid()
+		&& HierarchyView.IsValid())
 	{
-		UDMXPixelMappingBaseComponent* Source = ComponentDragDropOp->TryGetBaseComponent();
-
-		if (Source && HierarchyView.IsValid())
+		for (const FDMXPixelMappingDragDropOp::FDraggingComponentReference& DraggingComponentReference : ComponentDragDropOp->GetDraggingComponentReferences())
 		{
-			UDMXPixelMappingBaseComponent* Destination = Model.Pin()->GetReference().GetComponent();
-			
-			UDMXPixelMappingBaseComponent* NewParent = nullptr;
-			if (Source->CanBeMovedTo(Destination))
+			if (UDMXPixelMappingBaseComponent* Source = DraggingComponentReference.ComponentReference.GetComponent())
 			{
-				NewParent = Destination;
-			}
-			else if (Source->CanBeMovedTo(Destination->Parent))
-			{
-				NewParent = Destination->Parent;
-			}
-
-			if (NewParent)
-			{
-				NewParent->Modify();
-				Source->Modify();
-				Source->Parent->Modify();
-
-				Source->Parent->RemoveChild(Source);
-				NewParent->AddChild(Source);
-
-				HierarchyView->RequestComponentRedraw(NewParent);
-				HierarchyView->RequestRebuildTree();
+				UDMXPixelMappingBaseComponent* Destination = Model.Pin()->GetReference().GetComponent();
+
+				UDMXPixelMappingBaseComponent* NewParent = nullptr;
+				if (Source->CanBeMovedTo(Destination))
+				{
+					NewParent = Destination;
+				}
+				else if (Source->CanBeMovedTo(Destination->Parent))
+				{
+					NewParent = Destination->Parent;
+				}
+
+				if (NewParent)
+				{
+					NewParent->Modify();
+					Source->Modify();
+					Source->Parent->Modify();
+
+					Source->Parent->RemoveChild(Source);
+					NewParent->AddChild(Source);
+
+					HierarchyView->RequestComponentRedraw(NewParent);
+					HierarchyView->RequestRebuildTree();
+				}
 			}
-		}
+		}		
 	}
 	return FReply::Unhandled();
 }
diff --git a/Engine/Plugins/Experimental/DigitalMultiplex/DMXPixelMapping/Source/DMXPixelMappingRuntime/Private/Components/DMXPixelMappingFixtureGroupComponent.cpp b/Engine/Plugins/Experimental/DigitalMultiplex/DMXPixelMapping/Source/DMXPixelMappingRuntime/Private/Components/DMXPixelMappingFixtureGroupComponent.cpp
--- a/Engine/Plugins/Experimental/DigitalMultiplex/DMXPixelMapping/Source/DMXPixelMappingRuntime/Private/Components/DMXPixelMappingFixtureGroupComponent.cpp
+++ b/Engine/Plugins/Experimental/DigitalMultiplex/DMXPixelMapping/Source/DMXPixelMappingRuntime/Private/Components/DMXPixelMappingFixtureGroupComponent.cpp
@@ -207,7 +207,6 @@ TSharedRef<SWidget> UDMXPixelMappingFixtureGroupComponent::BuildSlot(TSharedRef<
 	Slot->Offset(FMargin(PositionX, PositionY, 0.f, 0.f));
 	CachedWidget->SetWidthOverride(SizeX);
 	CachedWidget->SetHeightOverride(SizeY);
-	CachedLabelBox->SetWidthOverride(SizeX);
 
 	UpdateWidget();
 
@@ -363,7 +362,6 @@ void UDMXPixelMappingFixtureGroupComponent::SetSizeWithinMinBoundaryBox()
 #if WITH_EDITOR
 	CachedWidget->SetWidthOverride(SizeX);
 	CachedWidget->SetHeightOverride(SizeY);
-	CachedLabelBox->SetWidthOverride(SizeX);
 #endif // WITH_EDITOR
 }
 
diff --git a/Engine/Plugins/Experimental/DigitalMultiplex/DMXPixelMapping/Source/DMXPixelMappingRuntime/Private/Components/DMXPixelMappingFixtureGroupItemComponent.cpp b/Engine/Plugins/Experimental/DigitalMultiplex/DMXPixelMapping/Source/DMXPixelMappingRuntime/Private/Components/DMXPixelMappingFixtureGroupItemComponent.cpp
--- a/Engine/Plugins/Experimental/DigitalMultiplex/DMXPixelMapping/Source/DMXPixelMappingRuntime/Private/Components/DMXPixelMappingFixtureGroupItemComponent.cpp
+++ b/Engine/Plugins/Experimental/DigitalMultiplex/DMXPixelMapping/Source/DMXPixelMappingRuntime/Private/Components/DMXPixelMappingFixtureGroupItemComponent.cpp
@@ -28,8 +28,8 @@ const FVector2D UDMXPixelMappingFixtureGroupItemComponent::MixPixelSize = FVecto
 UDMXPixelMappingFixtureGroupItemComponent::UDMXPixelMappingFixtureGroupItemComponent()
 	: DownsamplePixelIndex(0)
 {
-	SizeX = 100.f;
-	SizeY = 100.f;
+	SizeX = 10.f;
+	SizeY = 10.f;
 	PositionX = 0.f;
 	PositionY = 0.f;
 
@@ -93,46 +93,8 @@ void UDMXPixelMappingFixtureGroupItemComponent::PostParentAssigned()
 		}
 	}
 
-	// first compute a non-overlapping position
-	if (UDMXPixelMappingFixtureGroupComponent* FixtureGroupComponent = Cast<UDMXPixelMappingFixtureGroupComponent>(Parent))
-	{
-		bool bFirstStep = true;
-		for (UDMXPixelMappingBaseComponent* Component : FixtureGroupComponent->Children)
-		{
-			// skip myself
-			if (Component == this)
-			{
-				continue;
-			}
-			if (UDMXPixelMappingFixtureGroupItemComponent* FixtureGroupItem = Cast<UDMXPixelMappingFixtureGroupItemComponent>(Component))
-			{
-				if (bFirstStep)
-				{
-					PositionX = FixtureGroupItem->PositionX + FixtureGroupItem->SizeX;
-					PositionY = FixtureGroupItem->PositionY + FixtureGroupItem->SizeY;
-
-					bFirstStep = false;
-				}
-				else
-				{
-					PositionX = FMath::Max(FixtureGroupItem->PositionX + FixtureGroupItem->SizeX, PositionX);
-					PositionY = FMath::Max(FixtureGroupItem->PositionY + FixtureGroupItem->SizeY, PositionY);
-				}
-			}
-		}
-
-#if WITH_EDITOR
-		// Update relative position
-		RelativePositionX = PositionX - FixtureGroupComponent->GetPosition().X;
-		RelativePositionY = PositionY - FixtureGroupComponent->GetPosition().Y;
-
-		UpdateWidget();
-#endif // WITH_EDITOR
-	}
-
-	SetPositionInBoundaryBox(FVector2D(PositionX, PositionY));
-
 #if WITH_EDITOR
+	UpdateWidget();
 	AutoMapAttributes();
 #endif // WITH_EDITOR
 }
@@ -227,6 +189,7 @@ TSharedRef<SWidget> UDMXPixelMappingFixtureGroupItemComponent::BuildSlot(TShared
 
 	CachedLabelBox =
 		SNew(SBox)
+		.Padding(FMargin(2.f, 1.f, 2.f, 1.f))
 		.WidthOverride(SizeY)
 		.HAlign(HAlign_Left)
 		.VAlign(VAlign_Top)
@@ -247,8 +210,8 @@ TSharedRef<SWidget> UDMXPixelMappingFixtureGroupItemComponent::BuildSlot(TShared
 		.ZOrder(ZOrder)
 		[
 			SNew(SOverlay)
+			
 			+ SOverlay::Slot()
-			.Padding(TAttribute<FMargin>::Create(TAttribute<FMargin>::FGetter::CreateUObject(this, &UDMXPixelMappingFixtureGroupItemComponent::GetLabelPadding)))
 			.HAlign(HAlign_Fill)
 			.VAlign(VAlign_Fill)
 			[
@@ -685,13 +648,6 @@ void UDMXPixelMappingFixtureGroupItemComponent::SetSizeWithinBoundaryBox(const F
 	}
 }
 
-#if WITH_EDITOR
-FMargin UDMXPixelMappingFixtureGroupItemComponent::GetLabelPadding() const
-{
-	return FMargin(2.f, 1.f, 2.f, 0.f);
-}
-#endif // WITH_EDITOR
-
 #if WITH_EDITOR
 void UDMXPixelMappingFixtureGroupItemComponent::AutoMapAttributes()
 {
diff --git a/Engine/Plugins/Experimental/DigitalMultiplex/DMXPixelMapping/Source/DMXPixelMappingRuntime/Private/Components/DMXPixelMappingMatrixCellComponent.cpp b/Engine/Plugins/Experimental/DigitalMultiplex/DMXPixelMapping/Source/DMXPixelMappingRuntime/Private/Components/DMXPixelMappingMatrixCellComponent.cpp
--- a/Engine/Plugins/Experimental/DigitalMultiplex/DMXPixelMapping/Source/DMXPixelMappingRuntime/Private/Components/DMXPixelMappingMatrixCellComponent.cpp
+++ b/Engine/Plugins/Experimental/DigitalMultiplex/DMXPixelMapping/Source/DMXPixelMappingRuntime/Private/Components/DMXPixelMappingMatrixCellComponent.cpp
@@ -28,8 +28,8 @@ const FVector2D UDMXPixelMappingMatrixCellComponent::MixPixelSize = FVector2D(1.
 UDMXPixelMappingMatrixCellComponent::UDMXPixelMappingMatrixCellComponent()
 	: DownsamplePixelIndex(0)
 {
-	SizeX = 100.f;
-	SizeY = 100.f;
+	SizeX = 10.f;
+	SizeY = 10.f;
 
 #if WITH_EDITOR
 	RelativePositionX = 0.f;
diff --git a/Engine/Plugins/Experimental/DigitalMultiplex/DMXPixelMapping/Source/DMXPixelMappingRuntime/Public/Components/DMXPixelMappingFixtureGroupItemComponent.h b/Engine/Plugins/Experimental/DigitalMultiplex/DMXPixelMapping/Source/DMXPixelMappingRuntime/Public/Components/DMXPixelMappingFixtureGroupItemComponent.h
--- a/Engine/Plugins/Experimental/DigitalMultiplex/DMXPixelMapping/Source/DMXPixelMappingRuntime/Public/Components/DMXPixelMappingFixtureGroupItemComponent.h
+++ b/Engine/Plugins/Experimental/DigitalMultiplex/DMXPixelMapping/Source/DMXPixelMappingRuntime/Public/Components/DMXPixelMappingFixtureGroupItemComponent.h
@@ -86,10 +86,6 @@ private:
 	/** Set size of Fixture Pixel inside Fixture Group Boundary Box */
 	void SetSizeWithinBoundaryBox(const FVector2D& InSize);
 
-#if WITH_EDITOR
-	FMargin GetLabelPadding() const;
-#endif
-
 public:
 	UPROPERTY(EditAnywhere, BlueprintReadOnly, Category = "Selected Patch")
 	FDMXEntityFixturePatchRef FixturePatchRef;
