DMX - Fix top left corner of Fixture group doesn't aligned correctly with renderer

diff --git a/Engine/Plugins/VirtualProduction/DMX/DMXPixelMapping/Source/DMXPixelMappingEditorWidgets/Private/DMXPixelMappingComponentWidget.cpp b/Engine/Plugins/VirtualProduction/DMX/DMXPixelMapping/Source/DMXPixelMappingEditorWidgets/Private/DMXPixelMappingComponentWidget.cpp
--- a/Engine/Plugins/VirtualProduction/DMX/DMXPixelMapping/Source/DMXPixelMappingEditorWidgets/Private/DMXPixelMappingComponentWidget.cpp
+++ b/Engine/Plugins/VirtualProduction/DMX/DMXPixelMapping/Source/DMXPixelMappingEditorWidgets/Private/DMXPixelMappingComponentWidget.cpp
@@ -76,8 +76,11 @@ void FDMXPixelMappingComponentWidget::SetPosition(const FVector2D& LocalPosition
 {
 	if (OuterCanvas.IsValid() && ComponentSlot && LabelSlot)
 	{
-		ComponentSlot->SetOffset(FMargin(LocalPosition.X, LocalPosition.Y, 0.f, 0.f));
-		LabelSlot->SetOffset(FMargin(LocalPosition.X, LocalPosition.Y, 0.f, 0.f));
+		// In the middle of the top left pixel
+		const FVector2D UIPosition = FVector2D(LocalPosition.X + 0.5f, LocalPosition.Y + 0.5f);
+
+		ComponentSlot->SetOffset(FMargin(UIPosition.X, UIPosition.Y, 0.f, 0.f));
+		LabelSlot->SetOffset(FMargin(UIPosition.X, UIPosition.Y, 0.f, 0.f));
 	}
 }
 
