DMX - Fix user widget is marked as pending kill instead of letting it be garbage collected

diff --git a/Engine/Plugins/VirtualProduction/DMX/DMXPixelMapping/Source/DMXPixelMappingRuntime/Private/Components/DMXPixelMappingRendererComponent.cpp b/Engine/Plugins/VirtualProduction/DMX/DMXPixelMapping/Source/DMXPixelMappingRuntime/Private/Components/DMXPixelMappingRendererComponent.cpp
--- a/Engine/Plugins/VirtualProduction/DMX/DMXPixelMapping/Source/DMXPixelMappingRuntime/Private/Components/DMXPixelMappingRendererComponent.cpp
+++ b/Engine/Plugins/VirtualProduction/DMX/DMXPixelMapping/Source/DMXPixelMappingRuntime/Private/Components/DMXPixelMappingRendererComponent.cpp
@@ -265,7 +265,6 @@ void UDMXPixelMappingRendererComponent::OnMapChanged(UWorld* InWorld, EMapChange
 {
 	if (UserWidget != nullptr)
 	{
-		UserWidget->MarkAsGarbage();
 		UserWidget = nullptr;
 	}
 }
