DMX - Zooming in Pixel Mapping Designer panel now follows the cursor and snaps to correct places

diff --git a/Engine/Plugins/VirtualProduction/DMX/DMXPixelMapping/Source/DMXPixelMappingEditor/Private/Widgets/SDMXPixelMappingSurface.cpp b/Engine/Plugins/VirtualProduction/DMX/DMXPixelMapping/Source/DMXPixelMappingEditor/Private/Widgets/SDMXPixelMappingSurface.cpp
--- a/Engine/Plugins/VirtualProduction/DMX/DMXPixelMapping/Source/DMXPixelMappingEditor/Private/Widgets/SDMXPixelMappingSurface.cpp
+++ b/Engine/Plugins/VirtualProduction/DMX/DMXPixelMapping/Source/DMXPixelMappingEditor/Private/Widgets/SDMXPixelMappingSurface.cpp
@@ -88,9 +88,9 @@ struct FFixedZoomLevelsContainerDesignSurface : public FZoomLevelsContainer
 
 	int32 GetNearestZoomLevel(float InZoomAmount) const override
 	{
-		for ( int32 ZoomLevelIndex=0; ZoomLevelIndex < GetNumZoomLevels(); ++ZoomLevelIndex )
+		for (int32 ZoomLevelIndex = 0; ZoomLevelIndex < GetNumZoomLevels(); ++ZoomLevelIndex)
 		{
-			if ( InZoomAmount <= GetZoomAmount(ZoomLevelIndex) )
+			if (InZoomAmount <= GetZoomAmount(ZoomLevelIndex))
 			{
 				return ZoomLevelIndex;
 			}
@@ -130,8 +130,8 @@ struct FFixedZoomLevelsContainerDesignSurface : public FZoomLevelsContainer
 void SDMXPixelMappingSurface::Construct(const FArguments& InArgs, const TSharedPtr<FDMXPixelMappingToolkit>& InToolkit)
 {
 	ToolkitWeakPtr = InToolkit;
-	
-	if ( !ZoomLevels )
+
+	if (!ZoomLevels)
 	{
 		ZoomLevels = MakeUnique<FFixedZoomLevelsContainerDesignSurface>();
 	}
@@ -147,10 +147,10 @@ void SDMXPixelMappingSurface::Construct(const FArguments& InArgs, const TSharedP
 	bDrawGridLines = true;
 
 	ZoomLevelFade = FCurveSequence(0.0f, 1.0f);
-	ZoomLevelFade.Play( AsShared() );
+	ZoomLevelFade.Play(AsShared());
 
 	ZoomLevelGraphFade = FCurveSequence(0.0f, 0.5f);
-	ZoomLevelGraphFade.Play( AsShared() );
+	ZoomLevelGraphFade.Play(AsShared());
 
 	bDeferredZoomToExtents = false;
 
@@ -169,12 +169,12 @@ void SDMXPixelMappingSurface::Construct(const FArguments& InArgs, const TSharedP
 	ZoomStartOffset = FVector2D::ZeroVector;
 
 	ChildSlot
-	[
-		InArgs._Content.Widget
-	];
+		[
+			InArgs._Content.Widget
+		];
 }
 
-EActiveTimerReturnType SDMXPixelMappingSurface::HandleZoomToFit( double InCurrentTime, float InDeltaTime )
+EActiveTimerReturnType SDMXPixelMappingSurface::HandleZoomToFit(double InCurrentTime, float InDeltaTime)
 {
 	const FVector2D DesiredViewCenter = (ZoomTargetTopLeft + ZoomTargetBottomRight) * 0.5f;
 	const bool bDoneScrolling = ScrollToLocation(GetCachedGeometry(), DesiredViewCenter, bTeleportInsteadOfScrollingWhenZoomingToFit ? 1000.0f : InDeltaTime);
@@ -194,9 +194,9 @@ EActiveTimerReturnType SDMXPixelMappingSurface::HandleZoomToFit( double InCurren
 	return EActiveTimerReturnType::Continue;
 }
 
-void SDMXPixelMappingSurface::Tick( const FGeometry& AllottedGeometry, const double InCurrentTime, const float InDeltaTime )
+void SDMXPixelMappingSurface::Tick(const FGeometry& AllottedGeometry, const double InCurrentTime, const float InDeltaTime)
 {
-	if ( bDeferredZoomToExtents )
+	if (bDeferredZoomToExtents)
 	{
 		const FSlateRect Bounds = ComputeAreaBounds();
 		bDeferredZoomToExtents = false;
@@ -210,7 +210,7 @@ void SDMXPixelMappingSurface::Tick( const FGeometry& AllottedGeometry, const dou
 	}
 
 	if (!bRecenteredOnFirstTick)
-	{		
+	{
 		ZoomToFit(true);
 		bRecenteredOnFirstTick = true;
 	}
@@ -218,7 +218,7 @@ void SDMXPixelMappingSurface::Tick( const FGeometry& AllottedGeometry, const dou
 
 FCursorReply SDMXPixelMappingSurface::OnCursorQuery(const FGeometry& MyGeometry, const FPointerEvent& CursorEvent) const
 {
-	if ( bIsPanning )
+	if (bIsPanning)
 	{
 		return FCursorReply::Cursor(EMouseCursor::GrabHand);
 	}
@@ -245,7 +245,7 @@ FReply SDMXPixelMappingSurface::OnMouseButtonDown(const FGeometry& MyGeometry, c
 {
 	SCompoundWidget::OnMouseButtonDown(MyGeometry, MouseEvent);
 
-	if ( MouseEvent.GetEffectingButton() == EKeys::RightMouseButton || MouseEvent.GetEffectingButton() == EKeys::MiddleMouseButton )
+	if (MouseEvent.GetEffectingButton() == EKeys::RightMouseButton || MouseEvent.GetEffectingButton() == EKeys::MiddleMouseButton)
 	{
 		bIsPanning = false;
 
@@ -266,7 +266,7 @@ FReply SDMXPixelMappingSurface::OnMouseButtonUp(const FGeometry& MyGeometry, con
 {
 	SCompoundWidget::OnMouseButtonUp(MyGeometry, MouseEvent);
 
-	if ( MouseEvent.GetEffectingButton() == EKeys::RightMouseButton || MouseEvent.GetEffectingButton() == EKeys::MiddleMouseButton )
+	if (MouseEvent.GetEffectingButton() == EKeys::RightMouseButton || MouseEvent.GetEffectingButton() == EKeys::MiddleMouseButton)
 	{
 		bIsPanning = false;
 		bIsZooming = false;
@@ -282,12 +282,12 @@ FReply SDMXPixelMappingSurface::OnMouseMove(const FGeometry& MyGeometry, const F
 	const bool bIsMiddleMouseButtonDown = MouseEvent.IsMouseButtonDown(EKeys::MiddleMouseButton);
 	const FModifierKeysState ModifierKeysState = FSlateApplication::Get().GetModifierKeys();
 
-	if ( HasMouseCapture() )
+	if (HasMouseCapture())
 	{
 		const FVector2D CursorDelta = MouseEvent.GetCursorDelta();
 
 		const bool bShouldZoom = bIsRightMouseButtonDown && (bIsLeftMouseButtonDown || bIsMiddleMouseButtonDown || ModifierKeysState.IsAltDown() || FSlateApplication::Get().IsUsingTrackpad());
-		if ( bShouldZoom )
+		if (bShouldZoom)
 		{
 			const float MouseZoomScaling = 0.04f;
 			FReply ReplyState = FReply::Handled();
@@ -310,13 +310,13 @@ FReply SDMXPixelMappingSurface::OnMouseMove(const FGeometry& MyGeometry, const F
 
 			return ReplyState;
 		}
-		else if ( bIsRightMouseButtonDown || bIsMiddleMouseButtonDown )
+		else if (bIsRightMouseButtonDown || bIsMiddleMouseButtonDown)
 		{
 			FReply ReplyState = FReply::Handled();
 
 			bIsPanning = true;
 			bIsZooming = false;
-			ViewOffset = ViewOffsetStart + ( (MouseDownPositionAbsolute - MouseEvent.GetScreenSpacePosition()) / MyGeometry.Scale) / GetZoomAmount();
+			ViewOffset = ViewOffsetStart + ((MouseDownPositionAbsolute - MouseEvent.GetScreenSpacePosition()) / MyGeometry.Scale) / GetZoomAmount();
 
 			return ReplyState;
 		}
@@ -340,10 +340,10 @@ FReply SDMXPixelMappingSurface::OnTouchGesture(const FGeometry& MyGeometry, cons
 {
 	const EGestureEvent GestureType = GestureEvent.GetGestureType();
 	const FVector2D& GestureDelta = GestureEvent.GetGestureDelta();
-	if ( GestureType == EGestureEvent::Magnify )
+	if (GestureType == EGestureEvent::Magnify)
 	{
 		TotalGestureMagnify += GestureDelta.X;
-		if ( FMath::Abs(TotalGestureMagnify) > 0.07f )
+		if (FMath::Abs(TotalGestureMagnify) > 0.07f)
 		{
 			// We want to zoom into this point; i.e. keep it the same fraction offset into the panel
 			const FVector2D WidgetSpaceCursorPos = MyGeometry.AbsoluteToLocal(GestureEvent.GetScreenSpacePosition());
@@ -354,7 +354,7 @@ FReply SDMXPixelMappingSurface::OnTouchGesture(const FGeometry& MyGeometry, cons
 		}
 		return FReply::Handled();
 	}
-	else if ( GestureType == EGestureEvent::Scroll )
+	else if (GestureType == EGestureEvent::Scroll)
 	{
 		const EScrollGestureDirection DirectionSetting = GetDefault<ULevelEditorViewportSettings>()->ScrollGestureDirectionForOrthoViewports;
 		const bool bUseDirectionInvertdFromDevice = DirectionSetting == EScrollGestureDirection::Natural || (DirectionSetting == EScrollGestureDirection::UseSystemSetting && GestureEvent.IsDirectionInvertedFromDevice());
@@ -374,12 +374,12 @@ FReply SDMXPixelMappingSurface::OnTouchEnded(const FGeometry& MyGeometry, const
 
 inline float FancyMod(float Value, float Size)
 {
-	return ( ( Value >= 0 ) ? 0.0f : Size ) + FMath::Fmod(Value, Size);
+	return ((Value >= 0) ? 0.0f : Size) + FMath::Fmod(Value, Size);
 }
 
 float SDMXPixelMappingSurface::GetZoomAmount() const
 {
-	if ( AllowContinousZoomInterpolation.Get() )
+	if (AllowContinousZoomInterpolation.Get())
 	{
 		return FMath::Lerp(ZoomLevels->GetZoomAmount(PreviousZoomLevel), ZoomLevels->GetZoomAmount(ZoomLevel), ZoomLevelGraphFade.GetLerp());
 	}
@@ -399,13 +399,13 @@ void SDMXPixelMappingSurface::ChangeZoomLevel(int32 ZoomLevelDelta, const FVecto
 
 	const bool bAllowFullZoomRange =
 		// To zoom in past 1:1 the user must press control
-		( ZoomLevel == DefaultZoomLevel && ZoomLevelDelta > 0 && bOverrideZoomLimiting ) ||
+		(ZoomLevel == DefaultZoomLevel && ZoomLevelDelta > 0 && bOverrideZoomLimiting) ||
 		// If they are already zoomed in past 1:1, user may zoom freely
-		( ZoomLevel > DefaultZoomLevel );
+		(ZoomLevel > DefaultZoomLevel);
 
 	const float OldZoomLevel = ZoomLevel;
 
-	if ( bAllowFullZoomRange )
+	if (bAllowFullZoomRange)
 	{
 		ZoomLevel = FMath::Clamp(ZoomLevel + ZoomLevelDelta, 0, NumZoomLevels - 1);
 	}
@@ -415,27 +415,20 @@ void SDMXPixelMappingSurface::ChangeZoomLevel(int32 ZoomLevelDelta, const FVecto
 		ZoomLevel = FMath::Clamp(ZoomLevel + ZoomLevelDelta, 0, DefaultZoomLevel);
 	}
 
-	if ( OldZoomLevel != ZoomLevel )
+	if (OldZoomLevel != ZoomLevel)
 	{
 		PostChangedZoom();
 
 		// Note: This happens even when maxed out at a stop; so the user sees the animation and knows that they're at max zoom in/out
-		ZoomLevelFade.Play( AsShared() );
+		ZoomLevelFade.Play(AsShared());
 
 		// Re-center the screen so that it feels like zooming around the cursor.
 		{
-			FSlateRect GraphBounds = ComputeSensibleBounds();
-
-			// Make sure we are not zooming into/out into emptiness; otherwise the user will get lost..
-			const FVector2D ClampedPointToMaintainGraphSpace(
-				FMath::Clamp(PointToMaintainGraphSpace.X, GraphBounds.Left, GraphBounds.Right),
-				FMath::Clamp(PointToMaintainGraphSpace.Y, GraphBounds.Top, GraphBounds.Bottom)
-				);
-
-			const FVector2D NewViewOffset = ClampedPointToMaintainGraphSpace - WidgetSpaceZoomOrigin / GetZoomAmount();
+			const FVector2D NewViewOffset = PointToMaintainGraphSpace - WidgetSpaceZoomOrigin / GetZoomAmount();
 
 			// If we're panning while zooming we need to update the viewoffset start.
 			ViewOffsetStart += (NewViewOffset - ViewOffset);
+
 			// Update view offset to where ever we scrolled towards.
 			ViewOffset = NewViewOffset;
 
@@ -471,12 +464,12 @@ bool SDMXPixelMappingSurface::ScrollToLocation(const FGeometry& MyGeometry, FVec
 	ViewOffset = NewPosition - HalfOFScreenInGraphSpace;
 
 	// If within 1 pixel of target, stop interpolating
-	return ( ( NewPosition - DesiredCenterPosition ).SizeSquared() < 1.f );
+	return ((NewPosition - DesiredCenterPosition).SizeSquared() < 1.f);
 }
 
 bool SDMXPixelMappingSurface::ZoomToLocation(const FVector2D& CurrentSizeWithoutZoom, const FVector2D& InDesiredSize, bool bDoneScrolling)
 {
-	if ( bAllowContinousZoomInterpolation && ZoomLevelGraphFade.IsPlaying() )
+	if (bAllowContinousZoomInterpolation && ZoomLevelGraphFade.IsPlaying())
 	{
 		return false;
 	}
@@ -484,14 +477,14 @@ bool SDMXPixelMappingSurface::ZoomToLocation(const FVector2D& CurrentSizeWithout
 	const int32 DefaultZoomLevel = ZoomLevels->GetDefaultZoomLevel();
 	const int32 NumZoomLevels = ZoomLevels->GetNumZoomLevels();
 	int32 DesiredZoom = DefaultZoomLevel;
-	
+
 	// Find lowest zoom level that will display all nodes
-	for ( int32 Zoom = 0; Zoom < NumZoomLevels; ++Zoom )
+	for (int32 Zoom = 0; Zoom < NumZoomLevels; ++Zoom)
 	{
 		const FVector2D SizeWithZoom = (CurrentSizeWithoutZoom - ZoomToFitPadding) / ZoomLevels->GetZoomAmount(Zoom);
 		const FVector2D LeftOverSize = SizeWithZoom - InDesiredSize;
 
-		if ( ( InDesiredSize.X > SizeWithZoom.X ) || ( InDesiredSize.Y > SizeWithZoom.Y ) )
+		if ((InDesiredSize.X > SizeWithZoom.X) || (InDesiredSize.Y > SizeWithZoom.Y))
 		{
 			// Use the previous zoom level, this one is too tight
 			DesiredZoom = FMath::Max<int32>(0, Zoom - 1);
@@ -499,9 +492,9 @@ bool SDMXPixelMappingSurface::ZoomToLocation(const FVector2D& CurrentSizeWithout
 		}
 	}
 
-	if ( DesiredZoom != ZoomLevel )
+	if (DesiredZoom != ZoomLevel)
 	{
-		if ( bAllowContinousZoomInterpolation )
+		if (bAllowContinousZoomInterpolation)
 		{
 			// Animate to it
 			PreviousZoomLevel = ZoomLevel;
@@ -512,7 +505,7 @@ bool SDMXPixelMappingSurface::ZoomToLocation(const FVector2D& CurrentSizeWithout
 		else
 		{
 			// Do it instantly, either first or last
-			if ( DesiredZoom < ZoomLevel )
+			if (DesiredZoom < ZoomLevel)
 			{
 				// Zooming out; do it instantly
 				ZoomLevel = PreviousZoomLevel = DesiredZoom;
@@ -521,7 +514,7 @@ bool SDMXPixelMappingSurface::ZoomToLocation(const FVector2D& CurrentSizeWithout
 			else
 			{
 				// Zooming in; do it last
-				if ( bDoneScrolling )
+				if (bDoneScrolling)
 				{
 					ZoomLevel = PreviousZoomLevel = DesiredZoom;
 					ZoomLevelFade.Play(AsShared());
@@ -565,7 +558,7 @@ FVector2D SDMXPixelMappingSurface::GetViewOffset() const
 
 FVector2D SDMXPixelMappingSurface::GraphCoordToPanelCoord(const FVector2D& GraphSpaceCoordinate) const
 {
-	return ( GraphSpaceCoordinate - GetViewOffset() ) * GetZoomAmount();
+	return (GraphSpaceCoordinate - GetViewOffset()) * GetZoomAmount();
 }
 
 FVector2D SDMXPixelMappingSurface::PanelCoordToGraphCoord(const FVector2D& PanelSpaceCoordinate) const
@@ -599,7 +592,7 @@ void SDMXPixelMappingSurface::PaintBackgroundAsLines(const FSlateBrush* Backgrou
 
 	float ZoomFactor = RawZoomFactor;
 	float Inflation = 1.0f;
-	while ( ZoomFactor*Inflation*NominalGridSize <= GraphSmallestGridSize )
+	while (ZoomFactor * Inflation * NominalGridSize <= GraphSmallestGridSize)
 	{
 		Inflation *= 2.0f;
 	}
@@ -608,8 +601,8 @@ void SDMXPixelMappingSurface::PaintBackgroundAsLines(const FSlateBrush* Backgrou
 
 	FVector2D LocalGridOrigin = AllottedGeometry.AbsoluteToLocal(GridOrigin);
 
-	float ImageOffsetX = LocalGridOrigin.X - ((GridCellSize*RulePeriod) * FMath::Max(FMath::CeilToInt(LocalGridOrigin.X / (GridCellSize*RulePeriod)), 0));
-	float ImageOffsetY = LocalGridOrigin.Y - ((GridCellSize*RulePeriod) * FMath::Max(FMath::CeilToInt(LocalGridOrigin.Y / (GridCellSize*RulePeriod)), 0));
+	float ImageOffsetX = LocalGridOrigin.X - ((GridCellSize * RulePeriod) * FMath::Max(FMath::CeilToInt(LocalGridOrigin.X / (GridCellSize * RulePeriod)), 0));
+	float ImageOffsetY = LocalGridOrigin.Y - ((GridCellSize * RulePeriod) * FMath::Max(FMath::CeilToInt(LocalGridOrigin.Y / (GridCellSize * RulePeriod)), 0));
 
 	// Fill the background
 	FSlateDrawElement::MakeBox(
@@ -617,7 +610,7 @@ void SDMXPixelMappingSurface::PaintBackgroundAsLines(const FSlateBrush* Backgrou
 		DrawLayerId,
 		AllottedGeometry.ToPaintGeometry(),
 		BackgroundImage
-		);
+	);
 
 	if (bDrawGridLines)
 	{
