DMX - Improve PixelMapping Editor performance and related UX

diff --git a/Engine/Plugins/VirtualProduction/DMX/DMXPixelMapping/Source/DMXPixelMappingEditor/Private/Customizations/DMXPixelMappingDetailCustomization_FixtureGroup.cpp b/Engine/Plugins/VirtualProduction/DMX/DMXPixelMapping/Source/DMXPixelMappingEditor/Private/Customizations/DMXPixelMappingDetailCustomization_FixtureGroup.cpp
--- a/Engine/Plugins/VirtualProduction/DMX/DMXPixelMapping/Source/DMXPixelMappingEditor/Private/Customizations/DMXPixelMappingDetailCustomization_FixtureGroup.cpp
+++ b/Engine/Plugins/VirtualProduction/DMX/DMXPixelMapping/Source/DMXPixelMappingEditor/Private/Customizations/DMXPixelMappingDetailCustomization_FixtureGroup.cpp
@@ -20,6 +20,7 @@
 #include "IPropertyUtilities.h"
 #include "ScopedTransaction.h"
 #include "Layout/Visibility.h"
+#include "Misc/CoreDelegates.h"
 #include "Widgets/Views/SListView.h"
 #include "Widgets/Layout/SScrollBorder.h"
 #include "Widgets/SBoxPanel.h"
@@ -32,7 +33,7 @@ void FDMXPixelMappingDetailCustomization_FixtureGroup::CustomizeDetails(IDetailL
 	PropertyUtilities = InDetailLayout.GetPropertyUtilities();
 
 	WeakFixtureGroupComponent = GetSelectedFixtureGroupComponent(InDetailLayout);
-	if(UDMXPixelMappingFixtureGroupComponent* FixtureGroupComponent = WeakFixtureGroupComponent.Get())
+	if (UDMXPixelMappingFixtureGroupComponent* FixtureGroupComponent = WeakFixtureGroupComponent.Get())
 	{
 		// Listen to component changes
 		UDMXPixelMappingBaseComponent::GetOnComponentAdded().AddSP(this, &FDMXPixelMappingDetailCustomization_FixtureGroup::OnComponentAdded);
@@ -42,7 +43,7 @@ void FDMXPixelMappingDetailCustomization_FixtureGroup::CustomizeDetails(IDetailL
 		DMXLibraryHandle = InDetailLayout.GetProperty(GET_MEMBER_NAME_CHECKED(UDMXPixelMappingFixtureGroupComponent, DMXLibrary), UDMXPixelMappingFixtureGroupComponent::StaticClass());
 		DMXLibraryHandle->SetOnPropertyValueChanged(FSimpleDelegate::CreateSP(this, &FDMXPixelMappingDetailCustomization_FixtureGroup::OnLibraryChanged));
 		DMXLibraryHandle->SetOnChildPropertyValueChanged(FSimpleDelegate::CreateSP(this, &FDMXPixelMappingDetailCustomization_FixtureGroup::OnLibraryChanged));
-		
+
 		if (UDMXLibrary* DMXLibrary = GetSelectedDMXLibrary(FixtureGroupComponent))
 		{
 			UpdateFixturePatchesInUse(DMXLibrary);
@@ -67,7 +68,7 @@ void FDMXPixelMappingDetailCustomization_FixtureGroup::CustomizeDetails(IDetailL
 						{
 							return GroupItemComponent->FixturePatchRef.GetFixturePatch() == FixturePatch;
 						}
-						else if (UDMXPixelMappingMatrixComponent*  MatrixComponent = Cast<UDMXPixelMappingMatrixComponent>(BaseComponent))
+						else if (UDMXPixelMappingMatrixComponent* MatrixComponent = Cast<UDMXPixelMappingMatrixComponent>(BaseComponent))
 						{
 							return MatrixComponent->FixturePatchRef.GetFixturePatch() == FixturePatch;
 						}
@@ -94,7 +95,7 @@ void FDMXPixelMappingDetailCustomization_FixtureGroup::CustomizeDetails(IDetailL
 							FixturePatchDetailRowWidget
 						];
 				}
-			}		
+			}
 		}
 	}
 }
@@ -120,40 +121,37 @@ void FDMXPixelMappingDetailCustomization_FixtureGroup::OnLibraryChanged()
 					}
 				}
 				else if (UDMXPixelMappingMatrixComponent* ChildMatrix = Cast<UDMXPixelMappingMatrixComponent>(ChildComponent))
-							{
+				{
 					if (ChildMatrix->FixturePatchRef.GetFixturePatch() &&
 						ChildMatrix->FixturePatchRef.GetFixturePatch()->GetParentLibrary() != GroupComponent->DMXLibrary)
-						{
+					{
 						GroupComponent->RemoveChild(ChildMatrix);
-						}
 					}
+				}
 			};
 		}
-				}
+	}
 
-	if (!bRefreshing)
+	if (!RequestForceRefreshHandle.IsValid())
 	{
-		ForceRefresh();
-		bRefreshing = true;
+		RequestForceRefreshHandle = FCoreDelegates::OnEndFrame.AddSP(this, &FDMXPixelMappingDetailCustomization_FixtureGroup::ForceRefresh);
 	}
-			}
+}
 
 void FDMXPixelMappingDetailCustomization_FixtureGroup::OnComponentAdded(UDMXPixelMapping* PixelMapping, UDMXPixelMappingBaseComponent* Component)
 {
-	if (!bRefreshing)
+	if (!RequestForceRefreshHandle.IsValid())
 	{
-		ForceRefresh();
-		bRefreshing = true;
-		}
+		RequestForceRefreshHandle = FCoreDelegates::OnEndFrame.AddSP(this, &FDMXPixelMappingDetailCustomization_FixtureGroup::ForceRefresh);
 	}
+}
 
 void FDMXPixelMappingDetailCustomization_FixtureGroup::OnComponentRemoved(UDMXPixelMapping* PixelMapping, UDMXPixelMappingBaseComponent* Component)
 {
-	if (!bRefreshing)
+	if (!RequestForceRefreshHandle.IsValid())
 	{
-	ForceRefresh();
-		bRefreshing = true;
-}
+		RequestForceRefreshHandle = FCoreDelegates::OnEndFrame.AddSP(this, &FDMXPixelMappingDetailCustomization_FixtureGroup::ForceRefresh);
+	}
 }
 
 void FDMXPixelMappingDetailCustomization_FixtureGroup::ForceRefresh()
@@ -167,13 +165,13 @@ void FDMXPixelMappingDetailCustomization_FixtureGroup::ForceRefresh()
 		PropertyUtilities->ForceRefresh();
 	}
 
-	bRefreshing = false;
+	RequestForceRefreshHandle.Reset();
 }
 
 void FDMXPixelMappingDetailCustomization_FixtureGroup::OnFixturePatchLMBDown(const FGeometry& MyGeometry, const FPointerEvent& MouseEvent, FDMXEntityFixturePatchRef FixturePatchRef)
 {
 	UDMXEntityFixturePatch* FixturePatch = FixturePatchRef.GetFixturePatch();
-	if(FixturePatch)
+	if (FixturePatch)
 	{
 		if (SelectedFixturePatches.Num() == 0)
 		{
@@ -313,10 +311,10 @@ void FDMXPixelMappingDetailCustomization_FixtureGroup::UpdateFixturePatchesInUse
 		FixturePatchRef.SetEntity(FixturePatch);
 		FixturePatches.Add(FixturePatchRef);
 	}
-		
+
 	SelectedFixturePatches.RemoveAll([&FixturePatchesInLibrary](const FDMXEntityFixturePatchRef& SelectedFixturePatch) {
 		return !SelectedFixturePatch.GetFixturePatch() || !FixturePatchesInLibrary.Contains(SelectedFixturePatch.GetFixturePatch());
-	});
+		});
 }
 
 UDMXLibrary* FDMXPixelMappingDetailCustomization_FixtureGroup::GetSelectedDMXLibrary(UDMXPixelMappingFixtureGroupComponent* FixtureGroupComponent) const
diff --git a/Engine/Plugins/VirtualProduction/DMX/DMXPixelMapping/Source/DMXPixelMappingEditor/Private/Customizations/DMXPixelMappingDetailCustomization_FixtureGroup.h b/Engine/Plugins/VirtualProduction/DMX/DMXPixelMapping/Source/DMXPixelMappingEditor/Private/Customizations/DMXPixelMappingDetailCustomization_FixtureGroup.h
--- a/Engine/Plugins/VirtualProduction/DMX/DMXPixelMapping/Source/DMXPixelMappingEditor/Private/Customizations/DMXPixelMappingDetailCustomization_FixtureGroup.h
+++ b/Engine/Plugins/VirtualProduction/DMX/DMXPixelMapping/Source/DMXPixelMappingEditor/Private/Customizations/DMXPixelMappingDetailCustomization_FixtureGroup.h
@@ -102,6 +102,9 @@ private:
 	};
 	TArray<FDetailRowWidgetWithPatch> DetailRowWidgetsWithPatch;
 
+	/** Delegate handle while being bound to the OnEndFrame event */
+	FDelegateHandle RequestForceRefreshHandle;
+
 	TSharedPtr<IPropertyUtilities> PropertyUtilities;
 
 	/** If true, is refreshing */
diff --git a/Engine/Plugins/VirtualProduction/DMX/DMXPixelMapping/Source/DMXPixelMappingEditor/Private/Customizations/DMXPixelMappingDetailCustomization_FixtureGroupItem.cpp b/Engine/Plugins/VirtualProduction/DMX/DMXPixelMapping/Source/DMXPixelMappingEditor/Private/Customizations/DMXPixelMappingDetailCustomization_FixtureGroupItem.cpp
--- a/Engine/Plugins/VirtualProduction/DMX/DMXPixelMapping/Source/DMXPixelMappingEditor/Private/Customizations/DMXPixelMappingDetailCustomization_FixtureGroupItem.cpp
+++ b/Engine/Plugins/VirtualProduction/DMX/DMXPixelMapping/Source/DMXPixelMappingEditor/Private/Customizations/DMXPixelMappingDetailCustomization_FixtureGroupItem.cpp
@@ -246,7 +246,7 @@ TSharedRef<ITableRow> FDMXPixelMappingDetailCustomization_FixtureGroupItem::Gene
 }
 
 void FDMXPixelMappingDetailCustomization_FixtureGroupItem::CreateModulatorDetails(IDetailLayoutBuilder& InDetailLayout)
-{	
+{
 	IDetailCategoryBuilder& ModualtorsCategory = InDetailLayout.EditCategory("Modulators", LOCTEXT("DMXModulatorsCategory", "Modulators"), ECategoryPriority::Important);
 
 	TSharedPtr<IPropertyHandle> ModulatorClassesHandle = DetailLayout->GetProperty(GET_MEMBER_NAME_CHECKED(UDMXPixelMappingFixtureGroupItemComponent, ModulatorClasses), UDMXPixelMappingFixtureGroupItemComponent::StaticClass());
@@ -273,26 +273,26 @@ void FDMXPixelMappingDetailCustomization_FixtureGroupItem::CreateModulatorDetail
 					UClass* ModulatorClass = FirstGroupItemComponent->Modulators[IndexModulator]->GetClass();
 
 					for (const TWeakObjectPtr<UObject>& CustomizedObject : CustomizedObjects)
-				{
-						if (UDMXPixelMappingFixtureGroupItemComponent* GroupItemComponent = Cast<UDMXPixelMappingFixtureGroupItemComponent>(CustomizedObject.Get()))
 					{
+						if (UDMXPixelMappingFixtureGroupItemComponent* GroupItemComponent = Cast<UDMXPixelMappingFixtureGroupItemComponent>(CustomizedObject.Get()))
+						{
 							const bool bMultiEditableModulator =
 								GroupItemComponent->Modulators.IsValidIndex(IndexModulator) &&
 								GroupItemComponent->Modulators[IndexModulator] &&
 								GroupItemComponent->Modulators[IndexModulator]->GetClass() == ModulatorClass;
 
 							if (bMultiEditableModulator)
-						{
+							{
 								ModulatorsToEdit.Add(GroupItemComponent->Modulators[IndexModulator]);
-						}
+							}
 							else
-						{
-							// Don't allow multi edit if not all modulators are of same class
-							ModulatorsToEdit.Reset();
+							{
+								// Don't allow multi edit if not all modulators are of same class
+								ModulatorsToEdit.Reset();
+							}
 						}
 					}
 				}
-				}
 				else if (UDMXModulator* ModulatorOfFirstGroupItem = FirstGroupItemComponent->Modulators[IndexModulator])
 				{
 					ModulatorsToEdit.Add(ModulatorOfFirstGroupItem);
@@ -325,7 +325,7 @@ void FDMXPixelMappingDetailCustomization_FixtureGroupItem::CreateModulatorDetail
 						[
 							SNew(STextBlock)
 							.Text(LOCTEXT("ModulatorMultipleValues", "Multiple Values"))
-							.Font(IDetailLayoutBuilder::GetDetailFont())
+						.Font(IDetailLayoutBuilder::GetDetailFont())
 						];
 
 					break;
diff --git a/Engine/Plugins/VirtualProduction/DMX/DMXPixelMapping/Source/DMXPixelMappingEditor/Private/Customizations/DMXPixelMappingDetailCustomization_Matrix.cpp b/Engine/Plugins/VirtualProduction/DMX/DMXPixelMapping/Source/DMXPixelMappingEditor/Private/Customizations/DMXPixelMappingDetailCustomization_Matrix.cpp
--- a/Engine/Plugins/VirtualProduction/DMX/DMXPixelMapping/Source/DMXPixelMappingEditor/Private/Customizations/DMXPixelMappingDetailCustomization_Matrix.cpp
+++ b/Engine/Plugins/VirtualProduction/DMX/DMXPixelMapping/Source/DMXPixelMappingEditor/Private/Customizations/DMXPixelMappingDetailCustomization_Matrix.cpp
@@ -243,7 +243,7 @@ TSharedRef<ITableRow> FDMXPixelMappingDetailCustomization_Matrix::GenerateExpose
 }
 
 void FDMXPixelMappingDetailCustomization_Matrix::CreateModulatorDetails(IDetailLayoutBuilder& InDetailLayout)
-{	
+{
 	IDetailCategoryBuilder& ModualtorsCategory = InDetailLayout.EditCategory("Modulators", LOCTEXT("DMXModulatorsCategory", "Modulators"), ECategoryPriority::Important);
 
 	TSharedPtr<IPropertyHandle> ModulatorClassesHandle = DetailLayout->GetProperty(GET_MEMBER_NAME_CHECKED(UDMXPixelMappingMatrixComponent, ModulatorClasses), UDMXPixelMappingMatrixComponent::StaticClass());
@@ -270,27 +270,27 @@ void FDMXPixelMappingDetailCustomization_Matrix::CreateModulatorDetails(IDetailL
 					UClass* ModulatorClass = FirstMatrixComponent->Modulators[IndexModulator]->GetClass();
 
 					for (const TWeakObjectPtr<UObject>& CustomizedObject : CustomizedObjects)
-				{
-						if (UDMXPixelMappingMatrixComponent* MatrixComponent = Cast<UDMXPixelMappingMatrixComponent>(CustomizedObject.Get()))
 					{
+						if (UDMXPixelMappingMatrixComponent* MatrixComponent = Cast<UDMXPixelMappingMatrixComponent>(CustomizedObject.Get()))
+						{
 							const bool bMultiEditableModulator =
 								MatrixComponent->Modulators.IsValidIndex(IndexModulator) &&
 								MatrixComponent->Modulators[IndexModulator] &&
 								MatrixComponent->Modulators[IndexModulator]->GetClass() == ModulatorClass;
 
 							if (bMultiEditableModulator)
-						{
+							{
 								ModulatorsToEdit.Add(MatrixComponent->Modulators[IndexModulator]);
-						}
+							}
 							else
-						{
-							// Don't allow multi edit if not all modulators are of same class
-							ModulatorsToEdit.Reset();
+							{
+								// Don't allow multi edit if not all modulators are of same class
+								ModulatorsToEdit.Reset();
+							}
 						}
 					}
 				}
-				}
-				else if(UDMXModulator* ModulatorOfFirstMatrix = FirstMatrixComponent->Modulators[IndexModulator])
+				else if (UDMXModulator* ModulatorOfFirstMatrix = FirstMatrixComponent->Modulators[IndexModulator])
 				{
 					ModulatorsToEdit.Add(ModulatorOfFirstMatrix);
 				}
@@ -305,7 +305,7 @@ void FDMXPixelMappingDetailCustomization_Matrix::CreateModulatorDetails(IDetailL
 					DetailsViewArgs.bAllowSearch = false;
 					DetailsViewArgs.bHideSelectionTip = false;
 					DetailsViewArgs.NameAreaSettings = FDetailsViewArgs::HideNameArea;
-					
+
 					TSharedRef<IDetailsView> DetailsView = PropertyEditorModule.CreateDetailView(DetailsViewArgs);
 					DetailsView->SetObjects(ModulatorsToEdit);
 
diff --git a/Engine/Plugins/VirtualProduction/DMX/DMXPixelMapping/Source/DMXPixelMappingEditor/Private/Toolkits/DMXPixelMappingToolkit.cpp b/Engine/Plugins/VirtualProduction/DMX/DMXPixelMapping/Source/DMXPixelMappingEditor/Private/Toolkits/DMXPixelMappingToolkit.cpp
--- a/Engine/Plugins/VirtualProduction/DMX/DMXPixelMapping/Source/DMXPixelMappingEditor/Private/Toolkits/DMXPixelMappingToolkit.cpp
+++ b/Engine/Plugins/VirtualProduction/DMX/DMXPixelMapping/Source/DMXPixelMappingEditor/Private/Toolkits/DMXPixelMappingToolkit.cpp
@@ -33,6 +33,7 @@
 #include "Engine/TextureRenderTarget2D.h"
 #include "Framework/Commands/GenericCommands.h"
 #include "Kismet2/BlueprintEditorUtils.h"
+#include "Misc/ScopedSlowTask.h"
 #include "Widgets/SWidget.h"
 #include "Widgets/Docking/SDockTab.h"
 
@@ -260,7 +261,7 @@ void FDMXPixelMappingToolkit::Tick(float DeltaTime)
 		}
 	}
 
-	// Update component colors
+	// Update component widgets
 	UpdateComponentWidgetColors();
 }
 
@@ -292,8 +293,6 @@ void FDMXPixelMappingToolkit::HandleRemoveComponents()
 	FDMXEditorUtils::ClearAllDMXPortBuffers();
 	FDMXEditorUtils::ClearFixturePatchCachedData();
 
-	DesignerView->UpdateOutput(true);
-
 	UpdateBlueprintNodes(DMXPixelMapping);
 	OnComponentsAddedOrDeletedDelegate_DEPRECATED.Broadcast();
 
@@ -333,6 +332,7 @@ TArray<ComponentType> FDMXPixelMappingToolkit::MakeComponentArray(const TSet<FDM
 
 void FDMXPixelMappingToolkit::SelectComponents(const TSet<FDMXPixelMappingComponentReference>& InSelectedComponents)
 {
+	// Update selection
 	SelectedComponents.Empty();
 
 	SetActiveRenderComponent(nullptr);
@@ -360,12 +360,6 @@ void FDMXPixelMappingToolkit::SelectComponents(const TSet<FDMXPixelMappingCompon
 		}
 	}
 
-	if (SelectedComponents.Num() > 0)
-	{
-		// Force update design view
-		DesignerView->UpdateOutput(true);
-	}
-
 	OnSelectedComponentsChangedDelegate.Broadcast();
 }
 
@@ -537,33 +531,134 @@ void FDMXPixelMappingToolkit::OnSaveThumbnailImage()
 TArray<UDMXPixelMappingBaseComponent*> FDMXPixelMappingToolkit::CreateComponentsFromTemplates(UDMXPixelMappingRootComponent* RootComponent, UDMXPixelMappingBaseComponent* Target, const TArray<TSharedPtr<FDMXPixelMappingComponentTemplate>>& Templates)
 {
 	TArray<UDMXPixelMappingBaseComponent*> NewComponents;
-	if (ensureMsgf(RootComponent && Target, TEXT("Tried to create components from template but RootComponent or Target were invalid.")))
+	if (Templates.Num() > 0)
 	{
-		for (const TSharedPtr<FDMXPixelMappingComponentTemplate>& Template : Templates)
+		TGuardValue<bool>(bAddingComponents, true);
+
+		FScopedSlowTask SlowTask(Templates.Num(), LOCTEXT("CreateComponentsSlowTaskMessage", "Creating Pixel Mapping Components"));
+		SlowTask.MakeDialogDelayed(.5f);
+
+		if (ensureMsgf(RootComponent && Target, TEXT("Tried to create components from template but RootComponent or Target were invalid.")))
 		{
-			if (UDMXPixelMappingBaseComponent* NewComponent = Template->CreateComponent<UDMXPixelMappingBaseComponent>(RootComponent))
+			for (const TSharedPtr<FDMXPixelMappingComponentTemplate>& Template : Templates)
 			{
-				NewComponents.Add(NewComponent);
+				SlowTask.EnterProgressFrame(1);
+				if (UDMXPixelMappingBaseComponent* NewComponent = Template->CreateComponent<UDMXPixelMappingBaseComponent>(RootComponent))
+				{
+					NewComponents.Add(NewComponent);
 
-				Target->Modify();
-				NewComponent->Modify();
+					Target->Modify();
+					NewComponent->Modify();
 
-				Target->AddChild(NewComponent);
+					Target->AddChild(NewComponent);
+				}
 			}
 		}
+
+		UpdateBlueprintNodes(DMXPixelMapping);
 	}
 
 	return NewComponents;
 }
 
-void FDMXPixelMappingToolkit::OnComponentAdded(UDMXPixelMapping* PixelMapping, UDMXPixelMappingBaseComponent* InComponent)
+void FDMXPixelMappingToolkit::DeleteSelectedComponents()
 {
-	UpdateBlueprintNodes(DMXPixelMapping);
+	if (SelectedComponents.Num() > 0)
+	{
+		TGuardValue<bool>(bRemovingComponents, true);
+
+		// Count how many Descendants need to be removed 
+		const int32 NumComponentsAndItsDescendants = [this]()
+		{
+			int32 DescendantsCounter = 0;
+			for (const FDMXPixelMappingComponentReference& ComponentReference : SelectedComponents)
+			{
+				if (UDMXPixelMappingBaseComponent* Component = ComponentReference.GetComponent())
+				{
+					DescendantsCounter++;
+
+					constexpr bool bCountRecursive = true;
+					Component->ForEachChild([&DescendantsCounter](UDMXPixelMappingBaseComponent* ChildComponent)
+						{
+							DescendantsCounter++;
+						}, bCountRecursive);
+				}
+			}
+			return DescendantsCounter;
+		}();
+
+		// Make a slow task since removing may take a significant amount of time
+		const TSharedRef<FScopedSlowTask> SlowTask = MakeShared<FScopedSlowTask>(NumComponentsAndItsDescendants, LOCTEXT("DeleteComponentsSlowTaskMessage", "Removing Pixel Mapping Components"));
+		SlowTask->MakeDialogDelayed(.5f);
+
+		// Bind to the Component Removed Event and let the Slow Task update from that 
+		const FDelegateHandle RemoveComponentDelegateHandle = UDMXPixelMappingBaseComponent::GetOnComponentRemoved().AddRaw(this, &FDMXPixelMappingToolkit::UpdateRemoveComponentSlowTask, SlowTask);
+
+		// Remove the Components
+		const FScopedTransaction DeleteSelectedComponentsTransaction =
+			FScopedTransaction(FText::Format(LOCTEXT("DeleteSelectedComponents", "PixelMapping: Delete {0}|plural(one=Component, other=Components)"), SelectedComponents.Num()));
+
+		TSet<FDMXPixelMappingComponentReference> ParentComponentReferences;
+		for (const FDMXPixelMappingComponentReference& ComponentReference : SelectedComponents)
+		{
+			if (UDMXPixelMappingBaseComponent* Component = ComponentReference.GetComponent())
+			{
+				if (UDMXPixelMappingRendererComponent* RendererComponent = Cast<UDMXPixelMappingRendererComponent>(Component))
+				{
+					SetActiveRenderComponent(nullptr);
+				}
+
+				if (UDMXPixelMappingOutputComponent* OutputComponent = Cast<UDMXPixelMappingOutputComponent>(Component))
+				{
+					ActiveOutputComponents.Remove(OutputComponent);
+				}
+
+				if (UDMXPixelMappingBaseComponent* ParentComponent = Component->GetParent())
+				{
+					ParentComponentReferences.Add(GetReferenceFromComponent(ParentComponent));
+
+					Component->SetFlags(RF_Transactional);
+					Component->Modify();
+					ParentComponent->Modify();
+
+					ParentComponent->RemoveChild(ComponentReference.GetComponent());
+				}
+			}
+		}
+
+		// Unbind from the Component Removed Event
+		UDMXPixelMappingBaseComponent::GetOnComponentRemoved().Remove(RemoveComponentDelegateHandle);
+
+		// Select the Parent Component 
+		SelectComponents(ParentComponentReferences);
+
+		UpdateBlueprintNodes(DMXPixelMapping);
+	}
 }
 
-void FDMXPixelMappingToolkit::OnComponentRemoved(UDMXPixelMapping* PixelMapping, UDMXPixelMappingBaseComponent* InComponent)
+void FDMXPixelMappingToolkit::OnComponentAdded(UDMXPixelMapping* PixelMapping, UDMXPixelMappingBaseComponent* Component)
 {
-	UpdateBlueprintNodes(DMXPixelMapping);
+	// For performance reasons don't update blueprint nodes if the methods in this class add (possibly many) components.
+	// Instead call UpdateBlueprintNodes at the end such functions.
+	if (!bAddingComponents)
+	{
+		UpdateBlueprintNodes(DMXPixelMapping);
+	}
+}
+
+void FDMXPixelMappingToolkit::OnComponentRemoved(UDMXPixelMapping* PixelMapping, UDMXPixelMappingBaseComponent* Component)
+{
+	// For performance reasons don't update blueprint nodes if the methods in this class remove (possibly many) components .
+	// Instead call UpdateBlueprintNodes at the end such functions.
+	if (!bRemovingComponents)
+	{
+		UpdateBlueprintNodes(DMXPixelMapping);
+	}
+}
+
+void FDMXPixelMappingToolkit::UpdateRemoveComponentSlowTask(UDMXPixelMapping* PixelMapping, UDMXPixelMappingBaseComponent* InComponent, TSharedRef<FScopedSlowTask> SlowTask)
+{
+	SlowTask->EnterProgressFrame(1);
 }
 
 void FDMXPixelMappingToolkit::InitializeInternal(const EToolkitMode::Type Mode, const TSharedPtr<class IToolkitHost>& InitToolkitHost, const FGuid& MessageLogGuid)
@@ -728,45 +823,6 @@ void FDMXPixelMappingToolkit::CreateInternalViews()
 	CreateOrGetView_DetailsView();
 }
 
-void FDMXPixelMappingToolkit::DeleteSelectedComponents()
-{
-	const FScopedTransaction DeleteSelectedComponentsTransaction =
-		FScopedTransaction(FText::Format(LOCTEXT("DeleteSelectedComponents", "PixelMapping: Delete {0}|plural(one=Component, other=Components)"), SelectedComponents.Num()));
-
-	// Select parent as selected.
-	TSet<FDMXPixelMappingComponentReference> ParentComponentReferences;
-
-	for (const FDMXPixelMappingComponentReference& ComponentReference : SelectedComponents)
-	{
-		if (UDMXPixelMappingBaseComponent* Component = ComponentReference.GetComponent())
-		{
-			if (UDMXPixelMappingRendererComponent* RendererComponent = Cast<UDMXPixelMappingRendererComponent>(Component))
-			{
-				SetActiveRenderComponent(nullptr);
-			}
-
-			if (UDMXPixelMappingOutputComponent* OutputComponent = Cast<UDMXPixelMappingOutputComponent>(Component))
-			{
-				ActiveOutputComponents.Remove(OutputComponent);
-			}
-
-			if (UDMXPixelMappingBaseComponent* ParentComponent = Component->GetParent())
-			{
-				ParentComponentReferences.Add(GetReferenceFromComponent(ParentComponent));
-
-				Component->SetFlags(RF_Transactional);
-				Component->Modify();
-				ParentComponent->Modify();
-
-				ParentComponent->RemoveChild(ComponentReference.GetComponent());
-			}
-		}
-	}
-
-	// Should be called at the end of the function after Broadcast the changes
-	SelectComponents(ParentComponentReferences);
-}
-
 void FDMXPixelMappingToolkit::OnComponentRenamed(UDMXPixelMappingBaseComponent* InComponent)
 {
 	UpdateBlueprintNodes(GetDMXPixelMapping());
diff --git a/Engine/Plugins/VirtualProduction/DMX/DMXPixelMapping/Source/DMXPixelMappingEditor/Private/Toolkits/DMXPixelMappingToolkit.h b/Engine/Plugins/VirtualProduction/DMX/DMXPixelMapping/Source/DMXPixelMappingEditor/Private/Toolkits/DMXPixelMappingToolkit.h
--- a/Engine/Plugins/VirtualProduction/DMX/DMXPixelMapping/Source/DMXPixelMappingEditor/Private/Toolkits/DMXPixelMappingToolkit.h
+++ b/Engine/Plugins/VirtualProduction/DMX/DMXPixelMapping/Source/DMXPixelMappingEditor/Private/Toolkits/DMXPixelMappingToolkit.h
@@ -8,24 +8,23 @@
 #include "Toolkits/AssetEditorToolkit.h"
 #include "DMXPixelMappingComponentReference.h"
 
-class UDMXPixelMapping;
-class SDockableTab;
-class FTabManager;
-class SWidget;
-
+class FDMXPixelMappingPaletteViewModel;
+class FDMXPixelMappingToolbar;
 class SDMXPixelMappingPaletteView;
 class SDMXPixelMappingHierarchyView;
 class SDMXPixelMappingDesignerView;
 class SDMXPixelMappingPreviewView;
 class SDMXPixelMappingDetailsView;
-
-class FDMXPixelMappingPaletteViewModel;
-
+class UDMXPixelMapping;
 class UDMXPixelMappingBaseComponent;
-class FDMXPixelMappingToolbar;
-class UDMXPixelMappingRendererComponent;
-class UDMXPixelMappingOutputComponent;
 class UDMXPixelMappingMatrixComponent;
+class UDMXPixelMappingOutputComponent;
+class UDMXPixelMappingRendererComponent;
+
+struct FScopedSlowTask;
+class FTabManager;
+class SDockableTab;
+class SWidget;
 
 
 /**
@@ -146,8 +145,6 @@ public:
 
 	void AddRenderer();
 
-	void DeleteSelectedComponents();
-
 	void OnComponentRenamed(UDMXPixelMappingBaseComponent* InComponent);
 
 	/** 
@@ -156,13 +153,19 @@ public:
 	 */
 	TArray<UDMXPixelMappingBaseComponent*> CreateComponentsFromTemplates(UDMXPixelMappingRootComponent* RootComponent, UDMXPixelMappingBaseComponent* Target, const TArray<TSharedPtr<FDMXPixelMappingComponentTemplate>>& Templates);
 
-	UE_DEPRECATED(4.27, "Handled in UDMXPixelMappingMatrixComponent internally instead. No longer needs to be call explicitly.")
-	void DeleteMatrixPixels(UDMXPixelMappingMatrixComponent* InMatrixComponent) { checkNoEntry();}
-
-	UE_DEPRECATED(4.27, "Handled in UDMXPixelMappingMatrixComponent internally instead.  No longer needs to be call explicitly.")
-	void CreateMatrixPixels(UDMXPixelMappingMatrixComponent* InMatrixComponent) {}
+	/** Deletes the selected Components */
+	void DeleteSelectedComponents();
 
 private:
+	/** Called when a Component was added to the pixel mapping */
+	void OnComponentAdded(UDMXPixelMapping* PixelMapping, UDMXPixelMappingBaseComponent* Component);
+
+	/** Called when a Component was removed from the pixel mapping */
+	void OnComponentRemoved(UDMXPixelMapping* PixelMapping, UDMXPixelMappingBaseComponent* Component);
+
+	/** Updates the Slow Task when removing Components */
+	void UpdateRemoveComponentSlowTask(UDMXPixelMapping* PixelMapping, UDMXPixelMappingBaseComponent* Component, TSharedRef<FScopedSlowTask> SlowTask);
+
 	void PlayDMX();
 
 	void StopPlayingDMX();
@@ -193,13 +196,6 @@ private:
 
 	void CreateInternalViews();
 
-	/** Called when a component was added to the pixel mapping */
-	void OnComponentAdded(UDMXPixelMapping* PixelMapping, UDMXPixelMappingBaseComponent* InComponent);
-
-	/** Called when a component was removed from the pixel mapping */
-	void OnComponentRemoved(UDMXPixelMapping* PixelMapping, UDMXPixelMappingBaseComponent* InComponent);
-
-private:
 	UDMXPixelMapping* DMXPixelMapping;
 
 	/** List of open tool panels; used to ensure only one exists at any one time */
@@ -240,6 +236,12 @@ private:
 
 	static const uint8 RequestStopSendingMaxTicks;
 
+	/** True while adding components (to avoid needlessly updating blueprint nodes on each component added via our own methods) */
+	bool bAddingComponents = false;
+
+	/** True while removing components (to avoid needlessly updating blueprint nodes on each component removed via our own methods) */
+	bool bRemovingComponents = false;
+
 	FOnComponentsAddedOrDeletedDelegate OnComponentsAddedOrDeletedDelegate_DEPRECATED;
 
 public:
diff --git a/Engine/Plugins/VirtualProduction/DMX/DMXPixelMapping/Source/DMXPixelMappingEditor/Private/Views/SDMXPixelMappingDesignerView.cpp b/Engine/Plugins/VirtualProduction/DMX/DMXPixelMapping/Source/DMXPixelMappingEditor/Private/Views/SDMXPixelMappingDesignerView.cpp
--- a/Engine/Plugins/VirtualProduction/DMX/DMXPixelMapping/Source/DMXPixelMappingEditor/Private/Views/SDMXPixelMappingDesignerView.cpp
+++ b/Engine/Plugins/VirtualProduction/DMX/DMXPixelMapping/Source/DMXPixelMappingEditor/Private/Views/SDMXPixelMappingDesignerView.cpp
@@ -3,7 +3,9 @@
 #include "Views/SDMXPixelMappingDesignerView.h"
 
 #include "DMXPixelMapping.h"
+#include "DMXPixelMappingComponentWidget.h"
 #include "DMXPixelMappingEditorUtils.h"
+#include "SDMXPixelMappingComponentBox.h"
 #include "Components/DMXPixelMappingOutputComponent.h"
 #include "Components/DMXPixelMappingRendererComponent.h"
 #include "Components/DMXPixelMappingRootComponent.h"
@@ -14,16 +16,13 @@
 #include "DragDrop/DMXPixelMappingDragDropOp.h"
 #include "Library/DMXEntityFixturePatch.h"
 #include "Library/DMXLibrary.h"
-#include "Templates/DMXPixelMappingComponentTemplate.h"
-#include "Toolkits/DMXPixelMappingToolkit.h"
-
-#include "DMXPixelMappingComponentWidget.h"
-#include "SDMXPixelMappingComponentBox.h"
 #include "Widgets/SDMXPixelMappingDesignerCanvas.h"
-#include "Widgets/SDMXPixelMappingSourceTextureViewport.h"
 #include "Widgets/SDMXPixelMappingRuler.h"
-#include "Widgets/SDMXPixelMappingZoomPan.h"
+#include "Widgets/SDMXPixelMappingSourceTextureViewport.h"
 #include "Widgets/SDMXPixelMappingTransformHandle.h"
+#include "Widgets/SDMXPixelMappingZoomPan.h"
+#include "Templates/DMXPixelMappingComponentTemplate.h"
+#include "Toolkits/DMXPixelMappingToolkit.h"
 
 #include "Widgets/SCanvas.h"
 #include "Widgets/Text/STextBlock.h"
@@ -38,6 +37,11 @@
 
 #define LOCTEXT_NAMESPACE "SDMXPixelMappingDesignerView"
 
+SDMXPixelMappingDesignerView::~SDMXPixelMappingDesignerView()
+{
+	GEditor->UnregisterForUndo(this);
+}
+
 void SDMXPixelMappingDesignerView::Construct(const FArguments& InArgs, const TSharedPtr<FDMXPixelMappingToolkit>& InToolkit)
 {
 	SDMXPixelMappingSurface::Construct(SDMXPixelMappingSurface::FArguments()
@@ -126,7 +130,7 @@ void SDMXPixelMappingDesignerView::Construct(const FArguments& InArgs, const TSh
 						SNew(SDMXPixelMappingZoomPan)
 						.ZoomAmount(this, &SDMXPixelMappingDesignerView::GetZoomAmount)
 						.ViewOffset(this, &SDMXPixelMappingDesignerView::GetViewOffset)
-						.Visibility(this, &SDMXPixelMappingDesignerView::IsZoomPanVisible)
+						.Visibility(this, &SDMXPixelMappingDesignerView::GetZoomPanVisibility)
 
 						[
 							SNew(SOverlay)
@@ -143,7 +147,13 @@ void SDMXPixelMappingDesignerView::Construct(const FArguments& InArgs, const TSh
 
 							+ SOverlay::Slot()
 							[
-								SAssignNew(DesignCanvas, SDMXPixelMappingDesignerCanvas)
+								SAssignNew(DesignCanvasBorder, SBorder)
+								.HAlign(HAlign_Fill)
+								.VAlign(VAlign_Fill)
+								[
+									SAssignNew(DesignCanvas, SConstraintCanvas)
+									.Visibility(EVisibility::HitTestInvisible)
+								]
 							]
 						]
 					]
@@ -162,7 +172,60 @@ void SDMXPixelMappingDesignerView::Construct(const FArguments& InArgs, const TSh
 					.HAlign(HAlign_Fill)
 					.VAlign(VAlign_Fill)
 					[
-						CreateOverlayUI()
+						 SNew(SOverlay)
+						+ SOverlay::Slot()
+						.Padding(0)
+						+SOverlay::Slot()
+						.HAlign(HAlign_Fill)
+						.VAlign(VAlign_Top)
+						[
+							SNew(SHorizontalBox)
+
+							+ SHorizontalBox::Slot()
+							.AutoWidth()
+							.VAlign(VAlign_Center)
+							.Padding(6, 2, 0, 0)
+							[
+								SNew(STextBlock)
+								.TextStyle(FEditorStyle::Get(), "Graph.ZoomText")
+								.Text(this, &SDMXPixelMappingDesignerView::GetZoomText)
+								.ColorAndOpacity(this, &SDMXPixelMappingDesignerView::GetZoomTextColorAndOpacity)
+							]
+							+ SHorizontalBox::Slot()
+							.AutoWidth()
+							.VAlign(VAlign_Center)
+							.Padding(40, 2, 0, 0)
+							[
+								SNew(STextBlock)
+								.TextStyle(FEditorStyle::Get(), "Graph.ZoomText")
+								.Font(FCoreStyle::GetDefaultFontStyle(TEXT("BoldCondensed"), 14))
+								.Text(this, &SDMXPixelMappingDesignerView::GetCursorPositionText)
+								.ColorAndOpacity(FLinearColor(1.f, 1.f, 1.f, 0.25f))
+								.Visibility(this, &SDMXPixelMappingDesignerView::GetCursorPositionTextVisibility)
+							]
+							+ SHorizontalBox::Slot()
+							.FillWidth(1.0f)
+							[
+								SNew(SSpacer)
+								.Size(FVector2D(1, 1))
+							]
+
+							+ SHorizontalBox::Slot()
+							.AutoWidth()
+							.VAlign(VAlign_Center)
+							[
+								SNew(SButton)
+								.ButtonColorAndOpacity(FLinearColor::Transparent)
+								.ButtonStyle(FEditorStyle::Get(), "ViewportMenu.Button")
+								.ToolTipText(LOCTEXT("ZoomToFit_ToolTip", "Zoom To Fit"))
+								.OnClicked(this, &SDMXPixelMappingDesignerView::OnZoomToFitClicked)
+								.ContentPadding(FEditorStyle::Get().GetMargin("ViewportMenu.SToolBarButtonBlock.Button.Padding"))
+								[
+									SNew(SImage)
+									.Image(FEditorStyle::GetBrush("UMGEditor.ZoomToFit"))
+								]
+							]
+						]
 					]
 				]
 			]
@@ -179,6 +242,31 @@ void SDMXPixelMappingDesignerView::Construct(const FArguments& InArgs, const TSh
 	// Bind to component changes
 	UDMXPixelMappingBaseComponent::GetOnComponentAdded().AddSP(this, &SDMXPixelMappingDesignerView::OnComponentAdded);
 	UDMXPixelMappingBaseComponent::GetOnComponentRemoved().AddSP(this, &SDMXPixelMappingDesignerView::OnComponentRemoved);
+
+	GEditor->RegisterForUndo(this);
+}
+
+FOptionalSize SDMXPixelMappingDesignerView::GetPreviewAreaWidth() const
+{
+	FVector2D Area;
+	FVector2D Size;
+	GetPreviewAreaAndSize(Area, Size);
+
+	return Area.X;
+}
+
+FOptionalSize SDMXPixelMappingDesignerView::GetPreviewAreaHeight() const
+{
+	FVector2D Area;
+	FVector2D Size;
+	GetPreviewAreaAndSize(Area, Size);
+
+	return Area.Y;
+}
+
+float SDMXPixelMappingDesignerView::GetPreviewScale() const
+{
+	return GetZoomAmount();
 }
 
 FReply SDMXPixelMappingDesignerView::OnMouseButtonDown(const FGeometry& MyGeometry, const FPointerEvent& MouseEvent)
@@ -189,10 +277,8 @@ FReply SDMXPixelMappingDesignerView::OnMouseButtonDown(const FGeometry& MyGeomet
 	{
 		if (MouseEvent.GetEffectingButton() == EKeys::LeftMouseButton)
 		{
-			UDMXPixelMappingOutputComponent* ClickedComponent = FindComponentUnderCursor();
-
 			// Select and detect drag when something was clicked
-			if (ClickedComponent)
+			if (UDMXPixelMappingOutputComponent* ClickedComponent = GetComponentUnderCursor())
 			{
 				PendingSelectedComponent = ClickedComponent;
 
@@ -211,7 +297,7 @@ FReply SDMXPixelMappingDesignerView::OnMouseButtonDown(const FGeometry& MyGeomet
 							{
 								return true;
 							}
-					}
+						}
 						else if (UDMXPixelMappingFixtureGroupItemComponent* ClickedGroupItemComponent = Cast<UDMXPixelMappingFixtureGroupItemComponent>(ClickedComponent))
 						{
 							if (ClickedGroupItemComponent->GetParent() == SelectedComponentRef.GetComponent() &&
@@ -233,16 +319,16 @@ FReply SDMXPixelMappingDesignerView::OnMouseButtonDown(const FGeometry& MyGeomet
 				{
 					DragAnchor = GraphSpaceCursorPosition;
 
-				return
-					FReply::Handled()
-					.PreventThrottling()
-					.SetUserFocus(AsShared(), EFocusCause::Mouse)
-					.CaptureMouse(AsShared())
-					.DetectDrag(AsShared(), EKeys::LeftMouseButton);
+					return
+						FReply::Handled()
+						.PreventThrottling()
+						.SetUserFocus(AsShared(), EFocusCause::Mouse)
+						.CaptureMouse(AsShared())
+						.DetectDrag(AsShared(), EKeys::LeftMouseButton);
+				}
 			}
 		}
 	}
-	}
 
 	// Capture mouse for the drag handle and general mouse actions
 	return 
@@ -257,8 +343,8 @@ FReply SDMXPixelMappingDesignerView::OnMouseButtonUp(const FGeometry& MyGeometry
 	SDMXPixelMappingSurface::OnMouseButtonUp(MyGeometry, MouseEvent);
 
 	// Select the output component under the mouse
-	PendingSelectedComponent = FindComponentUnderCursor();
-	
+	PendingSelectedComponent = GetComponentUnderCursor();
+
 	// Select the Renderer if no Output Component was under the mouse
 	if (!PendingSelectedComponent.IsValid())
 	{
@@ -293,8 +379,7 @@ FReply SDMXPixelMappingDesignerView::OnMouseMove(const FGeometry& MyGeometry, co
 	if (MouseEvent.IsMouseButtonDown(EKeys::LeftMouseButton) && HasMouseCapture())
 	{
 		bool bIsRootWidgetSelected = false;
-		const TSet<FDMXPixelMappingComponentReference> SelectedComponentReferences = GetSelectedComponents();
-		for (const FDMXPixelMappingComponentReference& ComponentReference : SelectedComponentReferences)
+		for (const FDMXPixelMappingComponentReference& ComponentReference : GetSelectedComponents())
 		{
 			UDMXPixelMappingBaseComponent* Component = ComponentReference.GetComponent();
 			if (Component && Component->GetParent())
@@ -314,27 +399,17 @@ FReply SDMXPixelMappingDesignerView::OnMouseMove(const FGeometry& MyGeometry, co
 	return FReply::Unhandled();
 }
 
-void SDMXPixelMappingDesignerView::OnMouseEnter(const FGeometry& MyGeometry, const FPointerEvent& MouseEvent)
-{
-	SDMXPixelMappingSurface::OnMouseEnter(MyGeometry, MouseEvent);
-}
-
-void SDMXPixelMappingDesignerView::OnMouseLeave(const FPointerEvent& MouseEvent)
-{
-	SDMXPixelMappingSurface::OnMouseLeave(MouseEvent);
-}
-
 FReply SDMXPixelMappingDesignerView::OnKeyDown(const FGeometry& MyGeometry, const FKeyEvent& InKeyEvent)
 {
 	if (InKeyEvent.GetKey() == EKeys::Delete)
 	{
-		if (TSharedPtr<FDMXPixelMappingToolkit> ToolkitPtr = ToolkitWeakPtr.Pin())
+		if (const TSharedPtr<FDMXPixelMappingToolkit> ToolkitPtr = ToolkitWeakPtr.Pin())
 		{
 			ToolkitPtr->DeleteSelectedComponents();
 
 			return FReply::Handled();
-			}
 		}
+	}
 
 	return FReply::Unhandled();
 }
@@ -342,13 +417,9 @@ FReply SDMXPixelMappingDesignerView::OnKeyDown(const FGeometry& MyGeometry, cons
 void SDMXPixelMappingDesignerView::Tick(const FGeometry& AllottedGeometry, const double InCurrentTime, const float InDeltaTime)
 {
 	SDMXPixelMappingSurface::Tick(AllottedGeometry, InCurrentTime, InDeltaTime);
-	
-	if (TSharedPtr<FDMXPixelMappingToolkit> Toolkit = ToolkitWeakPtr.Pin())
-	{
-		CachedWidgetGeometry.Reset();
-		FArrangedWidget WindowWidgetGeometry(PreviewHitTestRoot.ToSharedRef(), GetDesignerGeometry());
-		PopulateWidgetGeometryCache(WindowWidgetGeometry);
 
+	if (const TSharedPtr<FDMXPixelMappingToolkit> Toolkit = ToolkitWeakPtr.Pin())
+	{
 		// Handle drag over on tick
 		if (PendingDragDropOp.IsValid())
 		{
@@ -372,8 +443,18 @@ void SDMXPixelMappingDesignerView::Tick(const FGeometry& AllottedGeometry, const
 			PendingDragDropOp.Reset();
 		}
 
-		// Update the preview
-		UpdateOutput(false);
+		// Update test if the renderer changed, and if so, rebuild the designer
+		UDMXPixelMappingRendererComponent* RendererComponent = Toolkit->GetActiveRendererComponent();
+		if (RendererComponent != CachedRendererComponent)
+		{
+			CachedRendererComponent = RendererComponent;
+			RebuildDesigner();
+		}
+
+		// Cache geometry
+		CachedWidgetGeometry.Reset();
+		FArrangedWidget WindowWidgetGeometry(PreviewHitTestRoot.ToSharedRef(), GetDesignerGeometry());
+		PopulateWidgetGeometryCache(WindowWidgetGeometry);
 
 		// Compute the origin in absolute space.
 		FGeometry RootGeometry = CachedWidgetGeometry.FindChecked(PreviewSizeConstraint.ToSharedRef()).Geometry;
@@ -412,10 +493,10 @@ FReply SDMXPixelMappingDesignerView::OnDragDetected(const FGeometry& MyGeometry,
 	{
 		TArray<UDMXPixelMappingOutputComponent*> DraggedComponentCandidates;
 
-		for (const FDMXPixelMappingComponentReference& SelectedComponent : GetSelectedComponents())
+		for (const FDMXPixelMappingComponentReference& SelectedComponent : SelectedComponents)
 		{
 			if (UDMXPixelMappingOutputComponent* OutputComponent = Cast<UDMXPixelMappingOutputComponent>(SelectedComponent.GetComponent()))
-			{				
+			{
 				DraggedComponentCandidates.Add(OutputComponent);
 			}
 		}
@@ -428,7 +509,7 @@ FReply SDMXPixelMappingDesignerView::OnDragDetected(const FGeometry& MyGeometry,
 			for (TWeakObjectPtr<UDMXPixelMappingBaseComponent> Parent = Candidate->GetParent(); Parent.IsValid(); Parent = Parent->GetParent())
 			{
 				bIsChild = DraggedComponentCandidates.ContainsByPredicate([Parent](const UDMXPixelMappingOutputComponent* BaseComponent)
-					{ 
+					{
 						return BaseComponent == Parent.Get();
 					});
 
@@ -450,17 +531,17 @@ FReply SDMXPixelMappingDesignerView::OnDragDetected(const FGeometry& MyGeometry,
 			// We rely on the first being the clicked component. There's an according comment in the DragDropOp and the detected drag is raised here.
 			UDMXPixelMappingOutputComponent* ClickedComponent = CastChecked<UDMXPixelMappingOutputComponent>(DraggedComponents[0]);
 
-			FVector2D GraphSpaceDragOffset = DragAnchor - ClickedComponent->GetPosition();
+			const FVector2D GraphSpaceDragOffset = DragAnchor - ClickedComponent->GetPosition();
 
-				TSharedRef<FDMXPixelMappingDragDropOp> DragDropOp = FDMXPixelMappingDragDropOp::New(GraphSpaceDragOffset, DraggedComponents);
-				DragDropOp->SetDecoratorVisibility(false);
+			TSharedRef<FDMXPixelMappingDragDropOp> DragDropOp = FDMXPixelMappingDragDropOp::New(GraphSpaceDragOffset, DraggedComponents);
+			DragDropOp->SetDecoratorVisibility(false);
 
-				// Clear any pending selected widgets
-				PendingSelectedComponent = nullptr;
+			// Clear any pending selected widgets
+			PendingSelectedComponent = nullptr;
 
-				return FReply::Handled().BeginDragDrop(DragDropOp);
-			}
+			return FReply::Handled().BeginDragDrop(DragDropOp);
 		}
+	}
 
 	return FReply::Handled();
 }
@@ -490,7 +571,7 @@ void SDMXPixelMappingDesignerView::OnDragLeave(const FDragDropEvent& DragDropEve
 
 	const TSharedPtr<FDMXPixelMappingDragDropOp> DragDropOp = DragDropEvent.GetOperationAs<FDMXPixelMappingDragDropOp>();
 	if (DragDropOp.IsValid())
-	{		
+	{
 		if (TSharedPtr<FDMXPixelMappingToolkit> ToolkitPtr = ToolkitWeakPtr.Pin())
 		{
 			FScopedRestoreSelection(ToolkitPtr.ToSharedRef(), StaticCastSharedRef<SDMXPixelMappingDesignerView>(AsShared()));
@@ -504,7 +585,7 @@ void SDMXPixelMappingDesignerView::OnDragLeave(const FDragDropEvent& DragDropEve
 					if (Component.IsValid() && Component->GetParent())
 					{
 						Parents.Add(ToolkitPtr->GetReferenceFromComponent(Component->GetParent()));
-					
+
 						Component->GetParent()->RemoveChild(Component.Get());
 					}
 				}
@@ -513,9 +594,7 @@ void SDMXPixelMappingDesignerView::OnDragLeave(const FDragDropEvent& DragDropEve
 				ToolkitPtr->SelectComponents(Parents);
 			}
 
-			// Update the preview
-			constexpr bool bForceUpdate = true;
-			UpdateOutput(bForceUpdate);
+			RebuildDesigner();
 		}
 	}
 }
@@ -527,9 +606,9 @@ FReply SDMXPixelMappingDesignerView::OnDragOver(const FGeometry& MyGeometry, con
 	const TSharedPtr<FDMXPixelMappingDragDropOp> DragDropOp = DragDropEvent.GetOperationAs<FDMXPixelMappingDragDropOp>();
 	if (DragDropOp.IsValid())
 	{
-			// Handle the drag drop op on tick
-			PendingDragDropOp = DragDropOp;
-		}
+		// Handle the drag drop op on tick
+		PendingDragDropOp = DragDropOp;
+	}
 
 	return FReply::Handled();
 }
@@ -538,203 +617,222 @@ FReply SDMXPixelMappingDesignerView::OnDrop(const FGeometry& MyGeometry, const F
 {
 	SDMXPixelMappingSurface::OnDrop(MyGeometry, DragDropEvent);
 
-	const TSharedPtr<FDMXPixelMappingDragDropOp> DragDropOp = DragDropEvent.GetOperationAs<FDMXPixelMappingDragDropOp>();
-	if (DragDropOp.IsValid())
-	{
-		if (TSharedPtr<FDMXPixelMappingToolkit> ToolkitPtr = ToolkitWeakPtr.Pin())
-		{
-			FScopedRestoreSelection(ToolkitPtr.ToSharedRef(), StaticCastSharedRef<SDMXPixelMappingDesignerView>(AsShared()));
-
-			CachedRendererComponent = nullptr;
-			UpdateOutput(false);
-			}
-		}
-
 	return FReply::Handled().EndDragDrop().SetUserFocus(AsShared());
 }
 
-UDMXPixelMappingOutputComponent* SDMXPixelMappingDesignerView::FindComponentUnderCursor() const
+FSlateRect SDMXPixelMappingDesignerView::ComputeAreaBounds() const
 {
-	UDMXPixelMappingOutputComponent* ComponentUnderCursor = nullptr;
-	if (UDMXPixelMappingRendererComponent* RendererComponent = CachedRendererComponent.Get())
+	FSlateRect Bounds = FSlateRect(0.f, 0.f, 0.f, 0.f);
+
+	const TSet<FDMXPixelMappingComponentReference> SelectedComponentReferences = GetSelectedComponents();
+	for (const FDMXPixelMappingComponentReference& ComponentReference : SelectedComponentReferences)
 	{
-		FVector2D GraphSpaceCursorPosition;
-		if (GetGraphSpaceCursorPosition(GraphSpaceCursorPosition))
+		if (UDMXPixelMappingBaseComponent* Component = ComponentReference.GetComponent())
 		{
-			RendererComponent->ForEachChild([&ComponentUnderCursor, &GraphSpaceCursorPosition](UDMXPixelMappingBaseComponent* InComponent)
-				{
-					if (UDMXPixelMappingOutputComponent* OutputComponent = Cast<UDMXPixelMappingOutputComponent>(InComponent))
-					{
-						if (!ComponentUnderCursor ||
-							ComponentUnderCursor->GetZOrder() <= OutputComponent->GetZOrder())
-						{
-							if (OutputComponent->IsVisible() &&
-								OutputComponent->IsOverPosition(GraphSpaceCursorPosition))
-							{
-								ComponentUnderCursor = OutputComponent;
-							}
-						}
-					}
-				}, true);
+			if (UDMXPixelMappingRendererComponent* RendererComponent = Component->GetRendererComponent())
+			{
+				Bounds.Left = FMath::Min(RendererComponent->GetPosition().X, Bounds.Left);
+				Bounds.Top = FMath::Min(RendererComponent->GetPosition().Y, Bounds.Top);
+
+				Bounds.Right = FMath::Max(RendererComponent->GetPosition().X + RendererComponent->GetSize().X, Bounds.Right);
+				Bounds.Bottom = FMath::Max(RendererComponent->GetPosition().Y + RendererComponent->GetSize().Y, Bounds.Bottom);
+
+				break;
+			}
 		}
 	}
 
-	return ComponentUnderCursor;
+	return Bounds;
 }
 
-bool SDMXPixelMappingDesignerView::GetGraphSpaceCursorPosition(FVector2D& OutGraphSpaceCursorPosition) const
+int32 SDMXPixelMappingDesignerView::GetGraphRulePeriod() const
 {
-	if (const FArrangedWidget* CachedPreviewSurface = CachedWidgetGeometry.Find(PreviewSizeConstraint.ToSharedRef()))
-	{
-		const FGeometry& RootGeometry = CachedPreviewSurface->Geometry;
-		OutGraphSpaceCursorPosition = RootGeometry.AbsoluteToLocal(FSlateApplication::Get().GetCursorPos());
+	return 10; // Parent override
+}
 
-		return true;
-	}
-	return false;
+float SDMXPixelMappingDesignerView::GetGridScaleAmount() const
+{
+	return 1.0f; // Parent override
 }
 
-void SDMXPixelMappingDesignerView::PopulateWidgetGeometryCache(FArrangedWidget& Root)
+int32 SDMXPixelMappingDesignerView::GetSnapGridSize() const
 {
-	const FSlateRect Rect = PreviewHitTestRoot->GetTickSpaceGeometry().GetLayoutBoundingRect();
-	const FSlateRect PaintRect = PreviewHitTestRoot->GetPaintSpaceGeometry().GetLayoutBoundingRect();
-	HittestGrid->SetHittestArea(Rect.GetTopLeft(), Rect.GetSize(), PaintRect.GetTopLeft());
-	HittestGrid->Clear();
+	return 4; // Parent override
+}
 
-	PopulateWidgetGeometryCache_Loop(Root);
+void SDMXPixelMappingDesignerView::PostUndo(bool bSuccess)
+{
+	CachedRendererComponent = nullptr;
 }
 
-void SDMXPixelMappingDesignerView::PopulateWidgetGeometryCache_Loop(FArrangedWidget& CurrentWidget)
+void SDMXPixelMappingDesignerView::PostRedo(bool bSuccess)
 {
-	bool bIncludeInHitTestGrid = true;
+	CachedRendererComponent = nullptr;
+}
 
-	if (bIncludeInHitTestGrid)
+void SDMXPixelMappingDesignerView::RebuildDesigner()
+{
+	if (DesignCanvas.IsValid())
 	{
-		HittestGrid->AddWidget(&(CurrentWidget.Widget.Get()), 0, 0, FSlateInvalidationWidgetSortOrder());
+		DesignCanvas->ClearChildren();
 	}
 
-	FArrangedChildren ArrangedChildren(EVisibility::All);
-	CurrentWidget.Widget->ArrangeChildren(CurrentWidget.Geometry, ArrangedChildren);
-
-	CachedWidgetGeometry.Add(CurrentWidget.Widget, CurrentWidget);
+	DesignCanvasBorder->SetContent
+	(
+		SAssignNew(DesignCanvas, SDMXPixelMappingDesignerCanvas)
+	);
+	const TSharedRef<SConstraintCanvas> ConstraintCanvas = SNew(SConstraintCanvas);
+	DesignCanvas->AddSlot()
+		[
+			ConstraintCanvas
+		];
 
-	for (int32 ChildIndex = 0; ChildIndex < ArrangedChildren.Num(); ++ChildIndex)
+	if (CachedRendererComponent.IsValid())
 	{
-		FArrangedWidget& SomeChild = ArrangedChildren[ChildIndex];
-		PopulateWidgetGeometryCache_Loop(SomeChild);
+		CachedRendererComponent->ForEachChild([ConstraintCanvas](UDMXPixelMappingBaseComponent* InComponent)
+			{
+				if (UDMXPixelMappingOutputComponent* Component = Cast<UDMXPixelMappingOutputComponent>(InComponent))
+				{
+					// Build all child DMX pixel mapping slots
+					Component->BuildSlot(ConstraintCanvas);
+				}
+			}, true);
 	}
 }
 
-FGeometry SDMXPixelMappingDesignerView::GetDesignerGeometry() const
-{
-	return PreviewHitTestRoot->GetTickSpaceGeometry();
-}
-
-void SDMXPixelMappingDesignerView::OnPaintBackground(const FGeometry& AllottedGeometry, const FSlateRect& MyCullingRect, FSlateWindowElementList& OutDrawElements, int32 LayerId) const
+void SDMXPixelMappingDesignerView::CreateExtensionWidgetsForSelection()
 {
-	SDMXPixelMappingSurface::OnPaintBackground(AllottedGeometry, MyCullingRect, OutDrawElements, LayerId);
-}
+	// Remove all the current extension widgets
+	ClearExtensionWidgets();
 
-FSlateRect SDMXPixelMappingDesignerView::ComputeAreaBounds() const
-{
-	FSlateRect Bounds = FSlateRect(0.f, 0.f, 0.f, 0.f);
+	// Get the selected widgets as an array
+	const TArray<FDMXPixelMappingComponentReference>& SelectedComponents = GetSelectedComponents().Array();
 
-	if (const TSharedPtr<FDMXPixelMappingToolkit> Toolkit = ToolkitWeakPtr.Pin())
+	if (SelectedComponents.Num() == 1)
 	{
-		const TSet<FDMXPixelMappingComponentReference> SelectedComponents = Toolkit->GetSelectedComponents();
-		for (const FDMXPixelMappingComponentReference& ComponentReference : SelectedComponents)
+		if (UDMXPixelMappingOutputComponent* OutputComponent = Cast<UDMXPixelMappingOutputComponent>(SelectedComponents[0].GetComponent()))
 		{
-			if (UDMXPixelMappingBaseComponent* Component = ComponentReference.GetComponent())
+			if (OutputComponent->IsVisible() && !OutputComponent->IsLockInDesigner())
 			{
-				if (UDMXPixelMappingRendererComponent* RendererComponent = Component->GetRendererComponent())
-				{
-					Bounds.Left = FMath::Min(RendererComponent->GetPosition().X, Bounds.Left);
-					Bounds.Top = FMath::Min(RendererComponent->GetPosition().Y, Bounds.Top);
-
-					Bounds.Right = FMath::Max(RendererComponent->GetPosition().X + RendererComponent->GetSize().X, Bounds.Right);
-					Bounds.Bottom = FMath::Max(RendererComponent->GetPosition().Y + RendererComponent->GetSize().Y, Bounds.Bottom);
+				// Add transform handles
+				constexpr float Offset = 10.f;
+				const TSharedPtr<SDMXPixelMappingDesignerView> Self = SharedThis(this);
+				TransformHandles.Add(SNew(SDMXPixelMappingTransformHandle, Self, EDMXPixelMappingTransformDirection::CenterRight, FVector2D(Offset, 0.f)));
+				TransformHandles.Add(SNew(SDMXPixelMappingTransformHandle, Self, EDMXPixelMappingTransformDirection::BottomCenter, FVector2D(0.f, Offset)));
+				TransformHandles.Add(SNew(SDMXPixelMappingTransformHandle, Self, EDMXPixelMappingTransformDirection::BottomRight, FVector2D(Offset, Offset)));
 
-					break;
+				// Add Widgets to designer surface
+				for (TSharedPtr<SDMXPixelMappingTransformHandle>& Handle : TransformHandles)
+				{
+					ExtensionWidgetCanvas->AddSlot()
+						.Position(TAttribute<FVector2D>::Create(TAttribute<FVector2D>::FGetter::CreateSP(this, &SDMXPixelMappingDesignerView::GetExtensionPosition, Handle)))
+						.Size(TAttribute<FVector2D>::Create(TAttribute<FVector2D>::FGetter::CreateSP(this, &SDMXPixelMappingDesignerView::GetExtensionSize, Handle)))
+						[
+							Handle.ToSharedRef()
+						];
 				}
 			}
 		}
 	}
-
-	return Bounds;
 }
 
-int32 SDMXPixelMappingDesignerView::GetGraphRulePeriod() const
+void SDMXPixelMappingDesignerView::ClearExtensionWidgets()
 {
-	return 10; // Parent override
+	ExtensionWidgetCanvas->ClearChildren();
 }
 
-float SDMXPixelMappingDesignerView::GetGridScaleAmount() const
+EVisibility SDMXPixelMappingDesignerView::GetExtensionCanvasVisibility() const
 {
-	return 1.0f; // Parent override
+	for (const FDMXPixelMappingComponentReference& Component : GetSelectedComponents())
+	{
+		UDMXPixelMappingOutputComponent* OutputComponent = Cast<UDMXPixelMappingOutputComponent>(Component.GetComponent());
+
+		if (!OutputComponent ||
+			!OutputComponent->IsVisible() ||
+			OutputComponent->IsLockInDesigner())
+		{
+			return EVisibility::Hidden;
+		}
+	}
+	return EVisibility::SelfHitTestInvisible;
 }
 
-int32 SDMXPixelMappingDesignerView::GetSnapGridSize() const
+FVector2D SDMXPixelMappingDesignerView::GetExtensionPosition(TSharedPtr<SDMXPixelMappingTransformHandle> Handle)
 {
-	return 4; // Parent override
+	const FDMXPixelMappingComponentReference& SelectedComponent = GetSelectedComponent();
+
+	if (SelectedComponent.IsValid())
+	{
+		FGeometry SelectedComponentGeometry;
+		FGeometry SelectedComponentParentGeometry;
+
+		if (GetComponentGeometry(SelectedComponent, SelectedComponentGeometry))
+		{
+			const FVector2D WidgetPosition = [Handle, &SelectedComponentGeometry]()
+			{
+				// Get the initial offset based on the location around the selected object.
+				switch (Handle->GetTransformDirection())
+				{
+				case EDMXPixelMappingTransformDirection::CenterRight:
+					return FVector2D(SelectedComponentGeometry.GetLocalSize().X, SelectedComponentGeometry.GetLocalSize().Y * 0.5f);
+				case EDMXPixelMappingTransformDirection::BottomLeft:
+					return FVector2D(0, SelectedComponentGeometry.GetLocalSize().Y);
+				case EDMXPixelMappingTransformDirection::BottomCenter:
+					return FVector2D(SelectedComponentGeometry.GetLocalSize().X * 0.5f, SelectedComponentGeometry.GetLocalSize().Y);
+				case EDMXPixelMappingTransformDirection::BottomRight:
+					return SelectedComponentGeometry.GetLocalSize();
+				default:
+					checkNoEntry(); // Unhandled enum value
+				}
+				return FVector2D::ZeroVector;
+			}();
+
+			const FVector2D SelectedWidgetScale = SelectedComponentGeometry.GetAccumulatedRenderTransform().GetMatrix().GetScale().GetVector();
+
+			const FVector2D ApplicationScaledOffset = Handle->GetOffset() * GetDesignerGeometry().Scale;
+
+			const FVector2D LocalOffsetFull = ApplicationScaledOffset / SelectedWidgetScale;
+			const FVector2D PositionFullOffset = GetDesignerGeometry().AbsoluteToLocal(SelectedComponentGeometry.LocalToAbsolute(WidgetPosition + LocalOffsetFull));
+			const FVector2D LocalOffsetHalf = (ApplicationScaledOffset / 2.0f) / SelectedWidgetScale;
+			const FVector2D PositionHalfOffset = GetDesignerGeometry().AbsoluteToLocal(SelectedComponentGeometry.LocalToAbsolute(WidgetPosition + LocalOffsetHalf));
+
+			FVector2D PivotCorrection = PositionHalfOffset - (PositionFullOffset + FVector2D(5.0f, 5.0f));
+
+			const FVector2D FinalPosition = PositionFullOffset + PivotCorrection;
+
+			return FinalPosition;
+		}
+	}
+
+	return FVector2D(0.f, 0.f);
 }
 
-FOptionalSize SDMXPixelMappingDesignerView::GetPreviewAreaWidth() const
+FVector2D SDMXPixelMappingDesignerView::GetExtensionSize(TSharedPtr<SDMXPixelMappingTransformHandle> Handle)
 {
-	FVector2D Area, Size;
-	GetPreviewAreaAndSize(Area, Size);
-
-	return Area.X;
+	return Handle->GetDesiredSize();
 }
 
-FOptionalSize SDMXPixelMappingDesignerView::GetPreviewAreaHeight() const
+FText SDMXPixelMappingDesignerView::GetHoveredComponentNameText() const
 {
-	FVector2D Area, Size;
-	GetPreviewAreaAndSize(Area, Size);
+	if (UDMXPixelMappingOutputComponent* ComponentUnderCursor = GetComponentUnderCursor())
+	{
+		return FText::FromString(ComponentUnderCursor->GetUserFriendlyName());
+	}
 
-	return Area.Y;
+	return FText();
 }
 
-void SDMXPixelMappingDesignerView::UpdateOutput(bool bForceUpdate)
+FText SDMXPixelMappingDesignerView::GetHoveredComponentParentNameText() const
 {
-	if (TSharedPtr<FDMXPixelMappingToolkit> Toolkit = ToolkitWeakPtr.Pin())
+	if (UDMXPixelMappingOutputComponent* ComponentUnderCursor = GetComponentUnderCursor())
 	{
-		UDMXPixelMappingRendererComponent* RendererComponent = Toolkit->GetActiveRendererComponent();
-		if (!CachedRendererComponent.IsValid())
+		if (UDMXPixelMappingBaseComponent* ParentOfComponentUnderCursor = ComponentUnderCursor->GetParent())
 		{
-			CachedRendererComponent = RendererComponent;
+			return FText::FromString(ParentOfComponentUnderCursor->GetUserFriendlyName());
 		}
-
-		if (RendererComponent)
-		{
-			if (bForceUpdate || RendererComponent != CachedRendererComponent)
-			{
-				DesignCanvas->ClearChildren();
-
-				DesignCanvas->AddSlot()
-					[
-						RendererComponent->TakeWidget()
-					];
-			}
-		}
-		else
-		{
-			DesignCanvas->ClearChildren();
-		}
-
-		CachedRendererComponent = RendererComponent;
 	}
-}
-
-void SDMXPixelMappingDesignerView::OnComponentAdded(UDMXPixelMapping* PixelMapping, UDMXPixelMappingBaseComponent* Component)
-{
-	CachedRendererComponent = nullptr;
-	UpdateOutput(true);
-}
 
-void SDMXPixelMappingDesignerView::OnComponentRemoved(UDMXPixelMapping* PixelMapping, UDMXPixelMappingBaseComponent* Component)
-{
-	CachedRendererComponent = nullptr;
-	UpdateOutput(true);
+	return FText();
 }
 
 EVisibility SDMXPixelMappingDesignerView::GetRulerVisibility() const
@@ -742,64 +840,6 @@ EVisibility SDMXPixelMappingDesignerView::GetRulerVisibility() const
 	return EVisibility::Visible;
 }
 
-TSharedRef<SWidget> SDMXPixelMappingDesignerView::CreateOverlayUI()
-{
-	return SNew(SOverlay)
-		+ SOverlay::Slot()
-		.Padding(0)
-		+SOverlay::Slot()
-		.HAlign(HAlign_Fill)
-		.VAlign(VAlign_Top)
-		[
-			SNew(SHorizontalBox)
-
-			+ SHorizontalBox::Slot()
-			.AutoWidth()
-			.VAlign(VAlign_Center)
-			.Padding(6, 2, 0, 0)
-			[
-				SNew(STextBlock)
-				.TextStyle(FEditorStyle::Get(), "Graph.ZoomText")
-				.Text(this, &SDMXPixelMappingDesignerView::GetZoomText)
-				.ColorAndOpacity(this, &SDMXPixelMappingDesignerView::GetZoomTextColorAndOpacity)
-			]
-			+ SHorizontalBox::Slot()
-			.AutoWidth()
-			.VAlign(VAlign_Center)
-			.Padding(40, 2, 0, 0)
-			[
-				SNew(STextBlock)
-				.TextStyle(FEditorStyle::Get(), "Graph.ZoomText")
-				.Font(FCoreStyle::GetDefaultFontStyle(TEXT("BoldCondensed"), 14))
-				.Text(this, &SDMXPixelMappingDesignerView::GetCursorPositionText)
-				.ColorAndOpacity(FLinearColor(1.f, 1.f, 1.f, 0.25f))
-				.Visibility(this, &SDMXPixelMappingDesignerView::GetCursorPositionTextVisibility)
-			]
-			+ SHorizontalBox::Slot()
-			.FillWidth(1.0f)
-			[
-				SNew(SSpacer)
-				.Size(FVector2D(1, 1))
-			]
-
-			+ SHorizontalBox::Slot()
-			.AutoWidth()
-			.VAlign(VAlign_Center)
-			[
-				SNew(SButton)
-				.ButtonColorAndOpacity(FLinearColor::Transparent)
-				.ButtonStyle(FEditorStyle::Get(), "ViewportMenu.Button")
-				.ToolTipText(LOCTEXT("ZoomToFit_ToolTip", "Zoom To Fit"))
-				.OnClicked(this, &SDMXPixelMappingDesignerView::HandleZoomToFitClicked)
-				.ContentPadding(FEditorStyle::Get().GetMargin("ViewportMenu.SToolBarButtonBlock.Button.Padding"))
-				[
-					SNew(SImage)
-					.Image(FEditorStyle::GetBrush("UMGEditor.ZoomToFit"))
-				]
-			]
-		];
-}
-
 FText SDMXPixelMappingDesignerView::GetCursorPositionText() const
 {
 	FVector2D CursorPosition;
@@ -815,59 +855,134 @@ EVisibility SDMXPixelMappingDesignerView::GetCursorPositionTextVisibility() cons
 	return IsHovered() ? EVisibility::SelfHitTestInvisible : EVisibility::Collapsed;
 }
 
-EVisibility SDMXPixelMappingDesignerView::IsZoomPanVisible() const
+EVisibility SDMXPixelMappingDesignerView::GetZoomPanVisibility() const
 {
-	if (TSharedPtr<FDMXPixelMappingToolkit> Toolkit = ToolkitWeakPtr.Pin())
+	if (UDMXPixelMappingRendererComponent* RendererComponent = CachedRendererComponent.Get())
 	{
-		if (UDMXPixelMappingRendererComponent* RendererComponent = Toolkit->GetActiveRendererComponent())
+		if (RendererComponent->GetRendererInputTexture())
 		{
-			if (RendererComponent->GetRendererInputTexture() != nullptr)
-			{
-				return EVisibility::Visible;
-			}
+			return EVisibility::Visible;
 		}
 	}
 
 	return EVisibility::Collapsed;
 }
 
-void SDMXPixelMappingDesignerView::GetPreviewAreaAndSize(FVector2D& Area, FVector2D& Size) const
+FReply SDMXPixelMappingDesignerView::OnZoomToFitClicked()
 {
-	check(SourceTextureViewport.IsValid());
+	ZoomToFit(/*bInstantZoom*/ false);
+	return FReply::Handled();
+}
 
-	Area = FVector2D(SourceTextureViewport->GetPreviewAreaWidth().Get(), SourceTextureViewport->GetPreviewAreaHeight().Get());
-	Size = Area;
+void SDMXPixelMappingDesignerView::OnComponentAdded(UDMXPixelMapping* PixelMapping, UDMXPixelMappingBaseComponent* Component)
+{
+	CachedRendererComponent = nullptr;
 }
 
-float SDMXPixelMappingDesignerView::GetPreviewScale() const
+void SDMXPixelMappingDesignerView::OnComponentRemoved(UDMXPixelMapping* PixelMapping, UDMXPixelMappingBaseComponent* Component)
 {
-	return GetZoomAmount();
+	CachedRendererComponent = nullptr;
 }
 
-FGeometry SDMXPixelMappingDesignerView::MakeGeometryWindowLocal(const FGeometry& WidgetGeometry) const
+void SDMXPixelMappingDesignerView::OnSelectedComponentsChanged()
 {
-	FGeometry NewGeometry = WidgetGeometry;
+	CreateExtensionWidgetsForSelection();
+}
 
-	TSharedPtr<SWindow> WidgetWindow = FSlateApplication::Get().FindWidgetWindow(SharedThis(this));
-	if (WidgetWindow.IsValid())
+void SDMXPixelMappingDesignerView::ResolvePendingSelectedComponents(bool bClearPreviousSelection)
+{
+	if (TSharedPtr<FDMXPixelMappingToolkit> ToolkitPtr = ToolkitWeakPtr.Pin())
 	{
-		TSharedRef<SWindow> CurrentWindowRef = WidgetWindow.ToSharedRef();
+		if (PendingSelectedComponent.IsValid())
+		{
+			if (UDMXPixelMappingMatrixCellComponent* MatrixCellComponent = Cast<UDMXPixelMappingMatrixCellComponent>(PendingSelectedComponent))
+			{
+				// If a Matrix Cell is selected, select the owning Matrix Component instead
+				PendingSelectedComponent = MatrixCellComponent->GetParent();
+			}
+			else if (UDMXPixelMappingFixtureGroupItemComponent* GroupItemComponent = Cast<UDMXPixelMappingFixtureGroupItemComponent>(PendingSelectedComponent))
+			{
+				// If an Item Component is selected and is lock in designer, select the owning Group Component instead
+				if (GroupItemComponent->IsLockInDesigner())
+				{
+					PendingSelectedComponent = GroupItemComponent->GetParent();
+				}
+			}
 
-		NewGeometry.AppendTransform(FSlateLayoutTransform(Inverse(CurrentWindowRef->GetPositionInScreen())));
+			TSet<FDMXPixelMappingComponentReference> SelectedComponents;
+
+			// Add the newly selected component first. This is important, e.g. for drag drop when iterating and using this as base
+			SelectedComponents.Add(ToolkitWeakPtr.Pin()->GetReferenceFromComponent(PendingSelectedComponent.Get()));
+
+			if (!bClearPreviousSelection)
+			{
+				SelectedComponents.Append(ToolkitWeakPtr.Pin()->GetSelectedComponents());
+			}
+			ToolkitWeakPtr.Pin()->SelectComponents(SelectedComponents);
+
+			PendingSelectedComponent = nullptr;
+		}
 	}
+}
 
-	return NewGeometry;
+void SDMXPixelMappingDesignerView::PopulateWidgetGeometryCache(FArrangedWidget& Root)
+{
+	const FSlateRect Rect = PreviewHitTestRoot->GetTickSpaceGeometry().GetLayoutBoundingRect();
+	const FSlateRect PaintRect = PreviewHitTestRoot->GetPaintSpaceGeometry().GetLayoutBoundingRect();
+	HittestGrid->SetHittestArea(Rect.GetTopLeft(), Rect.GetSize(), PaintRect.GetTopLeft());
+	HittestGrid->Clear();
+
+	PopulateWidgetGeometryCache_Loop(Root);
 }
 
-FReply SDMXPixelMappingDesignerView::HandleZoomToFitClicked()
+void SDMXPixelMappingDesignerView::PopulateWidgetGeometryCache_Loop(FArrangedWidget& CurrentWidget)
 {
-	ZoomToFit(/*bInstantZoom*/ false);
-	return FReply::Handled();
+	bool bIncludeInHitTestGrid = true;
+
+	if (bIncludeInHitTestGrid)
+	{
+		HittestGrid->AddWidget(&(CurrentWidget.Widget.Get()), 0, 0, FSlateInvalidationWidgetSortOrder());
+	}
+
+	FArrangedChildren ArrangedChildren(EVisibility::All);
+	CurrentWidget.Widget->ArrangeChildren(CurrentWidget.Geometry, ArrangedChildren);
+
+	CachedWidgetGeometry.Add(CurrentWidget.Widget, CurrentWidget);
+
+	for (int32 ChildIndex = 0; ChildIndex < ArrangedChildren.Num(); ++ChildIndex)
+	{
+		FArrangedWidget& SomeChild = ArrangedChildren[ChildIndex];
+		PopulateWidgetGeometryCache_Loop(SomeChild);
+	}
 }
 
-void SDMXPixelMappingDesignerView::OnSelectedComponentsChanged()
+UDMXPixelMappingOutputComponent* SDMXPixelMappingDesignerView::GetComponentUnderCursor() const
 {
-	CreateExtensionWidgetsForSelection();
+	UDMXPixelMappingOutputComponent* ComponentUnderCursor = nullptr;
+	if (UDMXPixelMappingRendererComponent* RendererComponent = CachedRendererComponent.Get())
+	{
+		FVector2D GraphSpaceCursorPosition;
+		if (GetGraphSpaceCursorPosition(GraphSpaceCursorPosition))
+		{
+			RendererComponent->ForEachChild([&ComponentUnderCursor, &GraphSpaceCursorPosition](UDMXPixelMappingBaseComponent* InComponent)
+				{
+					if (UDMXPixelMappingOutputComponent* OutputComponent = Cast<UDMXPixelMappingOutputComponent>(InComponent))
+					{
+						if (!ComponentUnderCursor ||
+							ComponentUnderCursor->GetZOrder() <= OutputComponent->GetZOrder())
+						{
+							if (OutputComponent->IsVisible() &&
+								OutputComponent->IsOverPosition(GraphSpaceCursorPosition))
+							{
+								ComponentUnderCursor = OutputComponent;
+							}
+						}
+					}
+				}, true);
+		}
+	}
+
+	return ComponentUnderCursor;
 }
 
 TSet<FDMXPixelMappingComponentReference> SDMXPixelMappingDesignerView::GetSelectedComponents() const
@@ -876,7 +991,7 @@ TSet<FDMXPixelMappingComponentReference> SDMXPixelMappingDesignerView::GetSelect
 	{
 		return ToolkitPtr->GetSelectedComponents();
 	}
-	
+
 	return TSet<FDMXPixelMappingComponentReference>();
 }
 
@@ -896,62 +1011,76 @@ FDMXPixelMappingComponentReference SDMXPixelMappingDesignerView::GetSelectedComp
 	return FDMXPixelMappingComponentReference();
 }
 
-FText SDMXPixelMappingDesignerView::GetHoveredComponentNameText() const
+bool SDMXPixelMappingDesignerView::GetComponentGeometry(FDMXPixelMappingComponentReference InComponentReference, FGeometry& OutGeometry)
 {
-	UDMXPixelMappingOutputComponent* OutputComponent = FindComponentUnderCursor();
-	if (OutputComponent)
+	if (UDMXPixelMappingBaseComponent* ComponentPreview = InComponentReference.GetComponent())
 	{
-		return FText::FromString(OutputComponent->GetUserFriendlyName());
+		return GetComponentGeometry(ComponentPreview, OutGeometry);
 	}
 
-	return FText();
+	return false;
 }
 
-FText SDMXPixelMappingDesignerView::GetHoveredComponentParentNameText() const
+bool SDMXPixelMappingDesignerView::GetComponentGeometry(UDMXPixelMappingBaseComponent* InBaseComponent, FGeometry& OutGeometry)
 {
-	UDMXPixelMappingOutputComponent* OutputComponent = FindComponentUnderCursor();
-	if (OutputComponent && OutputComponent->GetParent())
+	if (UDMXPixelMappingOutputComponent* OutputComponent = Cast<UDMXPixelMappingOutputComponent>(InBaseComponent))
 	{
-		return FText::FromString(OutputComponent->GetParent()->GetUserFriendlyName());
+		if (TSharedPtr<FDMXPixelMappingComponentWidget> ComponentWidget = OutputComponent->GetComponentWidget())
+		{
+			TSharedPtr<SWidget> CachedPreviewWidget = ComponentWidget->GetComponentBox();
+			if (CachedPreviewWidget.IsValid())
+			{
+				const FArrangedWidget* ArrangedWidget = CachedWidgetGeometry.Find(CachedPreviewWidget.ToSharedRef());
+				if (ArrangedWidget)
+				{
+					OutGeometry = ArrangedWidget->Geometry;
+					return true;
+				}
+			}
+		}
 	}
 
-	return FText();
+	return false;
 }
 
-void SDMXPixelMappingDesignerView::ResolvePendingSelectedComponents(bool bClearPreviousSelection)
+bool SDMXPixelMappingDesignerView::GetGraphSpaceCursorPosition(FVector2D& OutGraphSpaceCursorPosition) const
 {
-	if (TSharedPtr<FDMXPixelMappingToolkit> ToolkitPtr = ToolkitWeakPtr.Pin())
+	if (const FArrangedWidget* CachedPreviewSurface = CachedWidgetGeometry.Find(PreviewSizeConstraint.ToSharedRef()))
 	{
-		if (PendingSelectedComponent.IsValid())
-		{
-			if (UDMXPixelMappingMatrixCellComponent* MatrixCellComponent = Cast<UDMXPixelMappingMatrixCellComponent>(PendingSelectedComponent))
-			{
-				// If a Matrix Cell is selected, select the owning Matrix Component instead
-				PendingSelectedComponent = MatrixCellComponent->GetParent();
-			}
-			else if (UDMXPixelMappingFixtureGroupItemComponent* GroupItemComponent = Cast<UDMXPixelMappingFixtureGroupItemComponent>(PendingSelectedComponent))
-			{
-				// If an Item Component is selected and is lock in designer, select the owning Group Component instead
-				if (GroupItemComponent->IsLockInDesigner())
-				{
-					PendingSelectedComponent = GroupItemComponent->GetParent();
-				}
-			}
+		const FGeometry& RootGeometry = CachedPreviewSurface->Geometry;
+		OutGraphSpaceCursorPosition = RootGeometry.AbsoluteToLocal(FSlateApplication::Get().GetCursorPos());
 
-			TSet<FDMXPixelMappingComponentReference> SelectedComponents;
+		return true;
+	}
+	return false;
+}
 
-			// Add the newly selected component first. This is important, e.g. for drag drop when iterating and using this as base
-			SelectedComponents.Add(ToolkitWeakPtr.Pin()->GetReferenceFromComponent(PendingSelectedComponent.Get()));
+void SDMXPixelMappingDesignerView::GetPreviewAreaAndSize(FVector2D& Area, FVector2D& Size) const
+{
+	check(SourceTextureViewport.IsValid());
 
-			if (!bClearPreviousSelection)
-			{
-				SelectedComponents.Append(ToolkitWeakPtr.Pin()->GetSelectedComponents());
-			}
-			ToolkitWeakPtr.Pin()->SelectComponents(SelectedComponents);
+	Area = FVector2D(SourceTextureViewport->GetPreviewAreaWidth().Get(), SourceTextureViewport->GetPreviewAreaHeight().Get());
+	Size = Area;
+}
 
-			PendingSelectedComponent = nullptr;
-		}
+FGeometry SDMXPixelMappingDesignerView::GetDesignerGeometry() const
+{
+	return PreviewHitTestRoot->GetTickSpaceGeometry();
+}
+
+FGeometry SDMXPixelMappingDesignerView::MakeGeometryWindowLocal(const FGeometry& WidgetGeometry) const
+{
+	FGeometry NewGeometry = WidgetGeometry;
+
+	TSharedPtr<SWindow> WidgetWindow = FSlateApplication::Get().FindWidgetWindow(SharedThis(this));
+	if (WidgetWindow.IsValid())
+	{
+		TSharedRef<SWindow> CurrentWindowRef = WidgetWindow.ToSharedRef();
+
+		NewGeometry.AppendTransform(FSlateLayoutTransform(Inverse(CurrentWindowRef->GetPositionInScreen())));
 	}
+
+	return NewGeometry;
 }
 
 void SDMXPixelMappingDesignerView::HandleDragEnterFromDetailsOrPalette(const TSharedPtr<FDMXPixelMappingDragDropOp>& TemplateDragDropOp)
@@ -1016,7 +1145,7 @@ void SDMXPixelMappingDesignerView::HandleDragEnterFromDetailsOrPalette(const TSh
 					// Update selection depending on types
 					TSet<FDMXPixelMappingComponentReference> NewSelection;
 					if (NewMatrices.Num() > 0)
-					{					
+					{
 						// When only matrices and its childs are dragged, just select the matrix
 						for (UDMXPixelMappingBaseComponent* NewMatrix : NewMatrices)
 						{
@@ -1039,154 +1168,6 @@ void SDMXPixelMappingDesignerView::HandleDragEnterFromDetailsOrPalette(const TSh
 	}
 }
 
-void SDMXPixelMappingDesignerView::ClearExtensionWidgets()
-{
-	ExtensionWidgetCanvas->ClearChildren();
-}
-
-void SDMXPixelMappingDesignerView::CreateExtensionWidgetsForSelection()
-{
-	// Remove all the current extension widgets
-	ClearExtensionWidgets();
-
-	// Get the selected widgets as an array
-	const TArray<FDMXPixelMappingComponentReference>& Selected = GetSelectedComponents().Array();
-
-	constexpr int32 NumSelectedItems = 1;
-	if (Selected.Num() == NumSelectedItems)
-	{
-		if (UDMXPixelMappingOutputComponent* OutputComponent = Cast<UDMXPixelMappingOutputComponent>(Selected[0].GetComponent()))
-	{
-			if (OutputComponent->IsVisible() &&	!OutputComponent->IsLockInDesigner())
-			{
-	// Add transform handles
-				constexpr float Offset = 10.f;
-				const TSharedPtr<SDMXPixelMappingDesignerView> Self = SharedThis(this);
-				TransformHandles.Add(SNew(SDMXPixelMappingTransformHandle, Self, EDMXPixelMappingTransformDirection::CenterRight, FVector2D(Offset, 0.f)));
-				TransformHandles.Add(SNew(SDMXPixelMappingTransformHandle, Self, EDMXPixelMappingTransformDirection::BottomCenter, FVector2D(0.f, Offset)));
-	TransformHandles.Add(SNew(SDMXPixelMappingTransformHandle, Self, EDMXPixelMappingTransformDirection::BottomRight, FVector2D(Offset, Offset)));
-
-	// Add Widgets to designer surface
-	for (TSharedPtr<SDMXPixelMappingTransformHandle>& Handle : TransformHandles)
-	{
-		ExtensionWidgetCanvas->AddSlot()
-			.Position(TAttribute<FVector2D>::Create(TAttribute<FVector2D>::FGetter::CreateSP(this, &SDMXPixelMappingDesignerView::GetExtensionPosition, Handle)))
-			.Size(TAttribute<FVector2D>::Create(TAttribute<FVector2D>::FGetter::CreateSP(this, &SDMXPixelMappingDesignerView::GetExtensionSize, Handle)))
-			[
-				Handle.ToSharedRef()
-			];
-	}
-}
-		}
-	}
-}
-
-EVisibility SDMXPixelMappingDesignerView::GetExtensionCanvasVisibility() const
-{
-	for (const FDMXPixelMappingComponentReference& Component : GetSelectedComponents())
-	{
-		UDMXPixelMappingOutputComponent* OutputComponent = Cast<UDMXPixelMappingOutputComponent>(Component.GetComponent());
-
-		if (!OutputComponent ||
-			!OutputComponent->IsVisible() ||
-			OutputComponent->IsLockInDesigner())
-		{
-			return EVisibility::Hidden;
-		}
-	}
-	return EVisibility::SelfHitTestInvisible;
-}
-
-FVector2D SDMXPixelMappingDesignerView::GetExtensionPosition(TSharedPtr<SDMXPixelMappingTransformHandle> Handle)
-{
-	const FDMXPixelMappingComponentReference& SelectedComponent = GetSelectedComponent();
-
-	if (SelectedComponent.IsValid())
-	{
-		FGeometry SelectedComponentGeometry;
-		FGeometry SelectedComponentParentGeometry;
-
-		if (GetComponentGeometry(SelectedComponent, SelectedComponentGeometry))
-		{
-			FVector2D FinalPosition(0, 0);
-
-
-			FVector2D WidgetPosition;
-
-			// Get the initial offset based on the location around the selected object.
-			switch (Handle->GetTransformDirection())
-			{
-			case EDMXPixelMappingTransformDirection::CenterRight:
-				WidgetPosition = FVector2D(SelectedComponentGeometry.GetLocalSize().X, SelectedComponentGeometry.GetLocalSize().Y * 0.5f);
-				break;
-			case EDMXPixelMappingTransformDirection::BottomLeft:
-				WidgetPosition = FVector2D(0, SelectedComponentGeometry.GetLocalSize().Y);
-				break;
-			case EDMXPixelMappingTransformDirection::BottomCenter:
-				WidgetPosition = FVector2D(SelectedComponentGeometry.GetLocalSize().X * 0.5f, SelectedComponentGeometry.GetLocalSize().Y);
-				break;
-			case EDMXPixelMappingTransformDirection::BottomRight:
-				WidgetPosition = SelectedComponentGeometry.GetLocalSize();
-				break;
-			}
-
-			FVector2D SelectedWidgetScale = SelectedComponentGeometry.GetAccumulatedRenderTransform().GetMatrix().GetScale().GetVector();
-
-			FVector2D ApplicationScaledOffset = Handle->GetOffset() * GetDesignerGeometry().Scale;
-
-			FVector2D LocalOffsetFull = ApplicationScaledOffset / SelectedWidgetScale;
-			FVector2D PositionFullOffset = GetDesignerGeometry().AbsoluteToLocal(SelectedComponentGeometry.LocalToAbsolute(WidgetPosition + LocalOffsetFull));
-			FVector2D LocalOffsetHalf = (ApplicationScaledOffset / 2.0f) / SelectedWidgetScale;
-			FVector2D PositionHalfOffset = GetDesignerGeometry().AbsoluteToLocal(SelectedComponentGeometry.LocalToAbsolute(WidgetPosition + LocalOffsetHalf));
-
-			FVector2D PivotCorrection = PositionHalfOffset - (PositionFullOffset + FVector2D(5.0f, 5.0f));
-
-			FinalPosition = PositionFullOffset + PivotCorrection;
-
-			return FinalPosition;
-		}
-	}
-
-	return FVector2D(0, 0);
-}
-
-FVector2D SDMXPixelMappingDesignerView::GetExtensionSize(TSharedPtr<SDMXPixelMappingTransformHandle> Handle)
-{
-	return Handle->GetDesiredSize();
-}
-
-bool SDMXPixelMappingDesignerView::GetComponentGeometry(FDMXPixelMappingComponentReference ComponentReference, FGeometry& Geometry)
-{
-	if (UDMXPixelMappingBaseComponent* ComponentPreview = ComponentReference.GetComponent())
-	{
-		return GetComponentGeometry(ComponentPreview, Geometry);
-	}
-
-	return false;
-}
-
-bool SDMXPixelMappingDesignerView::GetComponentGeometry(UDMXPixelMappingBaseComponent* BaseComponent, FGeometry& Geometry)
-{
-	if (UDMXPixelMappingOutputComponent* OutputComponent = Cast<UDMXPixelMappingOutputComponent>(BaseComponent))
-	{
-		if (TSharedPtr<FDMXPixelMappingComponentWidget> ComponentWidget = OutputComponent->GetComponentWidget())
-		{
-			TSharedPtr<SWidget> CachedPreviewWidget = ComponentWidget->GetComponentBox();
-			if (CachedPreviewWidget.IsValid())
-			{
-				const FArrangedWidget* ArrangedWidget = CachedWidgetGeometry.Find(CachedPreviewWidget.ToSharedRef());
-				if (ArrangedWidget)
-				{
-					Geometry = ArrangedWidget->Geometry;
-					return true;
-				}
-			}
-		}
-	}
-	
-	return false;
-}
-
 SDMXPixelMappingDesignerView::FScopedRestoreSelection::FScopedRestoreSelection(TSharedRef<FDMXPixelMappingToolkit> ToolkitPtr, TSharedRef<SDMXPixelMappingDesignerView> DesignerView)
 	: WeakToolkit(ToolkitPtr)
 	, WeakDesignerView(DesignerView)
diff --git a/Engine/Plugins/VirtualProduction/DMX/DMXPixelMapping/Source/DMXPixelMappingEditor/Private/Views/SDMXPixelMappingDesignerView.h b/Engine/Plugins/VirtualProduction/DMX/DMXPixelMapping/Source/DMXPixelMappingEditor/Private/Views/SDMXPixelMappingDesignerView.h
--- a/Engine/Plugins/VirtualProduction/DMX/DMXPixelMapping/Source/DMXPixelMappingEditor/Private/Views/SDMXPixelMappingDesignerView.h
+++ b/Engine/Plugins/VirtualProduction/DMX/DMXPixelMapping/Source/DMXPixelMappingEditor/Private/Views/SDMXPixelMappingDesignerView.h
@@ -6,30 +6,33 @@
 
 #include "DMXPixelMappingComponentReference.h"
 
+#include "EditorUndoClient.h"
 #include "Templates/SharedPointer.h"
 
-class FDMXPixelMappingToolkit;
+class UDMXPixelMappingBaseComponent;
 class FDMXPixelMappingComponentReference;
 class FDMXPixelMappingDragDropOp;
-class SDMXPixelMappingRuler;
-class SDMXPixelMappingTransformHandle;
-class SDMXPixelMappingSourceTextureViewport;
-class UDMXPixelMappingBaseComponent;
 class UDMXPixelMappingFixtureGroupComponent;
 class UDMXPixelMappingFixtureGroupItemComponent;
 class UDMXPixelMappingOutputComponent;
 class UDMXPixelMappingRendererComponent;
+class SDMXPixelMappingRuler;
+class SDMXPixelMappingSourceTextureViewport;
+class FDMXPixelMappingToolkit;
+class SDMXPixelMappingTransformHandle;
 
-struct FOptionalSize;
 class FHittestGrid;
-class SCanvas;
-class SDMXPixelMappingDesignerCanvas;
+struct FOptionalSize;
+class SBorder;
 class SBox;
+class SCanvas;
+class SConstraintCanvas;
 class SOverlay;
 
 
 class SDMXPixelMappingDesignerView
 	: public SDMXPixelMappingSurface
+	, public FEditorUndoClient
 {
 private:
 	struct FComponentHitResult
@@ -49,7 +52,9 @@
 	SLATE_BEGIN_ARGS(SDMXPixelMappingDesignerView) { }
 	SLATE_END_ARGS()
 
-public:
+	/** Destructor */
+	virtual ~SDMXPixelMappingDesignerView();
+
 	/**
 	 * Constructs the widget.
 	 *
@@ -57,12 +62,20 @@
 	 */
 	void Construct(const FArguments& InArgs, const TSharedPtr<FDMXPixelMappingToolkit>& InToolkit);
 
+	/** The width of the preview screen for the UI */
+	FOptionalSize GetPreviewAreaWidth() const;
+
+	/** The height of the preview screen for the UI */
+	FOptionalSize GetPreviewAreaHeight() const;
+
+	/** Gets the scale of the preview screen */
+	float GetPreviewScale() const;
+
+protected:
 	// Begin SWidget interface
 	virtual FReply OnMouseButtonDown(const FGeometry& MyGeometry, const FPointerEvent& MouseEvent) override;
 	virtual FReply OnMouseButtonUp(const FGeometry& MyGeometry, const FPointerEvent& MouseEvent) override;
 	virtual FReply OnMouseMove(const FGeometry& MyGeometry, const FPointerEvent& MouseEvent) override;
-	virtual void OnMouseEnter(const FGeometry& MyGeometry, const FPointerEvent& MouseEvent) override;
-	virtual void OnMouseLeave(const FPointerEvent& MouseEvent) override;
 	virtual FReply OnKeyDown(const FGeometry& MyGeometry, const FKeyEvent& InKeyEvent) override;
 	virtual void Tick(const FGeometry& AllottedGeometry, const double InCurrentTime, const float InDeltaTime) override;
 	virtual FReply OnDragDetected(const FGeometry& MyGeometry, const FPointerEvent& MouseEvent) override;
@@ -73,89 +86,101 @@
 	// End of SWidget interface
 
 	//~ Begin SDMXPixelMappingSurface interface
-	virtual void OnPaintBackground(const FGeometry& AllottedGeometry, const FSlateRect& MyCullingRect, FSlateWindowElementList& OutDrawElements, int32 LayerId) const override;
 	virtual FSlateRect ComputeAreaBounds() const override;
 	virtual int32 GetGraphRulePeriod() const override;
 	virtual float GetGridScaleAmount() const override;
 	virtual int32 GetSnapGridSize() const override;
 	//~ End SDMXPixelMappingSurface interface
 
-	/** The width of the preview screen for the UI */
-	FOptionalSize GetPreviewAreaWidth() const;
-
-	/** The height of the preview screen for the UI */
-	FOptionalSize GetPreviewAreaHeight() const;
-
-	void UpdateOutput(bool bForceUpdate);
-
-	TSet<FDMXPixelMappingComponentReference> GetSelectedComponents() const;
-
-	FDMXPixelMappingComponentReference GetSelectedComponent() const;
-
-	float GetPreviewScale() const;
+	//~ Begin FEditorUndoClient Interface
+	virtual void PostUndo(bool bSuccess) override;
+	virtual void PostRedo(bool bSuccess) override;
+	// End of FEditorUndoClient
 
 private:
-	/** Returns the component under the cursor */
-	UDMXPixelMappingOutputComponent* FindComponentUnderCursor() const;
+	/** Rebuilds the designer widget */
+	void RebuildDesigner();
 
-	/** Returns the cursor position in graph space */
-	bool GetGraphSpaceCursorPosition(FVector2D& OutGraphSpaceCursorPosition) const;
+	/** Creates the extension widgets for the selected component */
+	void CreateExtensionWidgetsForSelection();
 
-	void PopulateWidgetGeometryCache(FArrangedWidget& Root);
+	/** Clears the extension widgets */
+	void ClearExtensionWidgets();
 
-	void PopulateWidgetGeometryCache_Loop(FArrangedWidget& Parent);
+	/** Returns if the extension widget canvas should be visible */
+	EVisibility GetExtensionCanvasVisibility() const;
 
-	FGeometry GetDesignerGeometry() const;
+	/** Gets the positioin of the extension widget */
+	FVector2D GetExtensionPosition(TSharedPtr<SDMXPixelMappingTransformHandle> Handle);
 
+	/** Gets the size of the extension widget */
+	FVector2D GetExtensionSize(TSharedPtr<SDMXPixelMappingTransformHandle> Handle);
+
+	/** Returns text for the currently hovered component */
+	FText GetHoveredComponentNameText() const;
+
+	/** Returns text for the currently hovered component's parent */
+	FText GetHoveredComponentParentNameText() const;
+
+	/** Returns the visibility of the side rulers */
+	EVisibility GetRulerVisibility() const;
+
+	/** Returns the cursor position as text */
+	FText GetCursorPositionText() const;
+
+	/** Returns the visibility of the cursor position text block */
+	EVisibility GetCursorPositionTextVisibility() const;
+
+	/** Returns the visibility of the ZoomPan widget */
+	EVisibility GetZoomPanVisibility() const;
+
+private:
+	/** Called when zoom to fit was clicked */
+	FReply OnZoomToFitClicked();
+
 	/** Called when a component was added */
 	void OnComponentAdded(UDMXPixelMapping* PixelMapping, UDMXPixelMappingBaseComponent* Component);
 
 	/** Called when a component was removed */
 	void OnComponentRemoved(UDMXPixelMapping* PixelMapping, UDMXPixelMappingBaseComponent* Component);
 
-	EVisibility GetRulerVisibility() const;
-
-	TSharedRef<SWidget> CreateOverlayUI();
-
-	FText GetCursorPositionText() const;
-
-	EVisibility GetCursorPositionTextVisibility() const;
-
-	EVisibility IsZoomPanVisible() const;
-
-	void GetPreviewAreaAndSize(FVector2D& Area, FVector2D& Size) const;
-
-	FGeometry MakeGeometryWindowLocal(const FGeometry& WidgetGeometry) const;
-
-	FReply HandleZoomToFitClicked();
-
+	/** Called when the selected component in the toolkit changed */
 	void OnSelectedComponentsChanged();
 
 	/** Adds any pending selected Components to the selection set */
 	void ResolvePendingSelectedComponents(bool bClearPreviousSelection = false);
-	
+
+	void PopulateWidgetGeometryCache(FArrangedWidget& Root);
+
+	void PopulateWidgetGeometryCache_Loop(FArrangedWidget& Parent);
+
+private:
+	/** Returns the component under the cursor */
+	UDMXPixelMappingOutputComponent* GetComponentUnderCursor() const;
+
+	/** Returns the selected components */
+	TSet<FDMXPixelMappingComponentReference> GetSelectedComponents() const;
+
+	/** Returns the first best selected component */
+	FDMXPixelMappingComponentReference GetSelectedComponent() const;
+
+	/** Gets the Geometry of the component, returns true if successful */
+	bool GetComponentGeometry(FDMXPixelMappingComponentReference InComponentReference, FGeometry& OutGeometry);
+
+	/** Gets the Geometry of the component, returns true if successful */
+	bool GetComponentGeometry(UDMXPixelMappingBaseComponent* InBaseComponent, FGeometry& OutGeometry);
+
+	/** Returns the cursor position in graph space */
+	bool GetGraphSpaceCursorPosition(FVector2D& OutGraphSpaceCursorPosition) const;
+
+	void GetPreviewAreaAndSize(FVector2D& Area, FVector2D& Size) const;
+
+	FGeometry GetDesignerGeometry() const;
+
+	FGeometry MakeGeometryWindowLocal(const FGeometry& WidgetGeometry) const;
+
 	/** Updates the drag drop op to use the patches selected in the Palette */
 	void HandleDragEnterFromDetailsOrPalette(const TSharedPtr<FDMXPixelMappingDragDropOp>& DragDropOp);
 
-	void ClearExtensionWidgets();
-
-	void CreateExtensionWidgetsForSelection();
-
-	EVisibility GetExtensionCanvasVisibility() const;
-
-	FVector2D GetExtensionPosition(TSharedPtr<SDMXPixelMappingTransformHandle> Handle);
-
-	FVector2D GetExtensionSize(TSharedPtr<SDMXPixelMappingTransformHandle> Handle);
-
-	bool GetComponentGeometry(FDMXPixelMappingComponentReference ComponentReference, FGeometry& Geometry);
-
-	bool GetComponentGeometry(UDMXPixelMappingBaseComponent* BaseComponent, FGeometry& Geometry);
-
-	/** Returns text for the currently hovered component */
-	FText GetHoveredComponentNameText() const;
-
-	/** Returns text for the currently hovered component's parent */
-	FText GetHoveredComponentParentNameText() const;
-
 private:
 	/** The component that should be selected, but isn't selected yet */
@@ -164,20 +189,25 @@
 	/** The drag over DragDropOp that should be handled on the next tick*/
 	TSharedPtr<FDMXPixelMappingDragDropOp> PendingDragDropOp;
 
+	/** Canvas that holds the extension widget */
 	TSharedPtr<SCanvas> ExtensionWidgetCanvas;
 
+	/** Widget that holds the source texture */
 	TSharedPtr<SDMXPixelMappingSourceTextureViewport> SourceTextureViewport;
 
 	/** Canvas that holds the component widgets */
-	TSharedPtr<SDMXPixelMappingDesignerCanvas> DesignCanvas;
+	TSharedPtr<SConstraintCanvas> DesignCanvas;
 
 	TSharedPtr<SBox> PreviewSizeConstraint;
 
 	TSharedPtr<SOverlay> PreviewHitTestRoot;
 
+	/** Borer that contains the Design Canvas */
+	TSharedPtr<SBorder> DesignCanvasBorder;
+
 	TSharedPtr<FHittestGrid> HittestGrid;
 
-	TMap<TSharedRef<SWidget>, FArrangedWidget> CachedWidgetGeometry;
+	TMap<TWeakPtr<SWidget>, FArrangedWidget> CachedWidgetGeometry;
 
 	TWeakObjectPtr<UDMXPixelMappingRendererComponent> CachedRendererComponent;
 
@@ -205,4 +235,3 @@
 		TArray<TWeakObjectPtr<UDMXPixelMappingBaseComponent>> CachedSelectedComponents;
 	};
 };
-
diff --git a/Engine/Plugins/VirtualProduction/DMX/DMXPixelMapping/Source/DMXPixelMappingEditor/Private/Widgets/SDMXPixelMappingFixturePatchDetailRow.cpp b/Engine/Plugins/VirtualProduction/DMX/DMXPixelMapping/Source/DMXPixelMappingEditor/Private/Widgets/SDMXPixelMappingFixturePatchDetailRow.cpp
--- a/Engine/Plugins/VirtualProduction/DMX/DMXPixelMapping/Source/DMXPixelMappingEditor/Private/Widgets/SDMXPixelMappingFixturePatchDetailRow.cpp
+++ b/Engine/Plugins/VirtualProduction/DMX/DMXPixelMapping/Source/DMXPixelMappingEditor/Private/Widgets/SDMXPixelMappingFixturePatchDetailRow.cpp
@@ -12,16 +12,16 @@
 
 #define LOCTEXT_NAMESPACE "SDMXPixelMappingPreviewView"
 
-namespace
+namespace UE::DMX::Private::SDMXPixelMappingFixturePatchDetailRow
 {
 	constexpr FLinearColor NormalBGColor(0.f, 0.0f, 0.f, 0.f); // Fully transparent so it just shows the background
 	constexpr FLinearColor HighlightBGColor(0.87f, 0.64f, 0.f, 0.5f); // Yellow transparent
-	constexpr FLinearColor ErrorBGColor(1.f, 0.f, 0.f, 0.5f); // Yellow transparent
-
 }
 
 void SDMXPixelMappingFixturePatchDetailRow::Construct(const FArguments& InArgs)
 {
+	using namespace UE::DMX::Private::SDMXPixelMappingFixturePatchDetailRow;
+
 	OnLMBDown = InArgs._OnLMBDown;
 	OnLMBUp   = InArgs._OnLMBUp;
 	OnDragged = InArgs._OnDragged;
@@ -36,8 +36,8 @@ void SDMXPixelMappingFixturePatchDetailRow::Construct(const FArguments& InArgs)
 		+ SOverlay::Slot()
 		.Padding(-3.f) // Need to overdraw to avoid having gaps between the detail rows
 		[
-			SAssignNew(Border, SBorder)
-			.BorderImage(FEditorStyle::GetBrush("DetailsView.GroupSection"))
+			SNew(SBorder)
+			.BorderImage(FEditorStyle::GetBrush("WhiteBrush"))
 			.BorderBackgroundColor_Lambda([this]()
 				{
 					return bHighlight ? HighlightBGColor : NormalBGColor;
diff --git a/Engine/Plugins/VirtualProduction/DMX/DMXPixelMapping/Source/DMXPixelMappingEditor/Private/Widgets/SDMXPixelMappingFixturePatchDetailRow.h b/Engine/Plugins/VirtualProduction/DMX/DMXPixelMapping/Source/DMXPixelMappingEditor/Private/Widgets/SDMXPixelMappingFixturePatchDetailRow.h
--- a/Engine/Plugins/VirtualProduction/DMX/DMXPixelMapping/Source/DMXPixelMappingEditor/Private/Widgets/SDMXPixelMappingFixturePatchDetailRow.h
+++ b/Engine/Plugins/VirtualProduction/DMX/DMXPixelMapping/Source/DMXPixelMappingEditor/Private/Widgets/SDMXPixelMappingFixturePatchDetailRow.h
@@ -12,7 +12,6 @@
 class UDMXEntityFixturePatch;
 
 class IPropertyHandle;
-class SBorder;
 class STextBlock;
 
 
@@ -61,9 +60,6 @@ private:
 	/** Error text, shows if not empty */
 	FText ErrorText;
 	
-	/** The border around the widget */
-	TSharedPtr<SBorder> Border;
-
 	// Slate Args
 	FDMXOnFixturePatchDetailRowMouseButtonEvent OnLMBDown;
 	FDMXOnFixturePatchDetailRowMouseButtonEvent OnLMBUp;
diff --git a/Engine/Plugins/VirtualProduction/DMX/DMXPixelMapping/Source/DMXPixelMappingEditor/Private/Widgets/SDMXPixelMappingHierarchyItem.cpp b/Engine/Plugins/VirtualProduction/DMX/DMXPixelMapping/Source/DMXPixelMappingEditor/Private/Widgets/SDMXPixelMappingHierarchyItem.cpp
--- a/Engine/Plugins/VirtualProduction/DMX/DMXPixelMapping/Source/DMXPixelMappingEditor/Private/Widgets/SDMXPixelMappingHierarchyItem.cpp
+++ b/Engine/Plugins/VirtualProduction/DMX/DMXPixelMapping/Source/DMXPixelMappingEditor/Private/Widgets/SDMXPixelMappingHierarchyItem.cpp
@@ -106,8 +106,14 @@ FReply SDMXPixelMappingHierarchyItem::OnDropWidget(const FDragDropEvent& InDragD
 					NewParent = Destination->GetParent();
 				}
 
-				if (NewParent)
+				const bool bOldParentIsNewParent = NewParent == Source->GetParent();
+				if (NewParent && !bOldParentIsNewParent)
 				{
+					// Add to the new parent
+					NewParent->Modify();
+					Source->Modify();
+					NewParent->AddChild(Source);
+
 					// Remove from the old parent
 					if (Source->GetParent())
 					{
@@ -115,11 +121,6 @@ FReply SDMXPixelMappingHierarchyItem::OnDropWidget(const FDragDropEvent& InDragD
 						Source->GetParent()->RemoveChild(Source);
 					}
 
-					// Add to the new parent
-					NewParent->Modify();
-					Source->Modify();
-					NewParent->AddChild(Source);
-
 					// Adopt location and size of new parent if required
 					if (UDMXPixelMappingOutputComponent* ParentOutputComponent = Cast<UDMXPixelMappingOutputComponent>(NewParent))
 					{
diff --git a/Engine/Plugins/VirtualProduction/DMX/DMXPixelMapping/Source/DMXPixelMappingEditor/Private/Widgets/SDMXPixelMappingPaletteItem.cpp b/Engine/Plugins/VirtualProduction/DMX/DMXPixelMapping/Source/DMXPixelMappingEditor/Private/Widgets/SDMXPixelMappingPaletteItem.cpp
--- a/Engine/Plugins/VirtualProduction/DMX/DMXPixelMapping/Source/DMXPixelMappingEditor/Private/Widgets/SDMXPixelMappingPaletteItem.cpp
+++ b/Engine/Plugins/VirtualProduction/DMX/DMXPixelMapping/Source/DMXPixelMappingEditor/Private/Widgets/SDMXPixelMappingPaletteItem.cpp
@@ -1,16 +1,24 @@
 // Copyright Epic Games, Inc. All Rights Reserved.
 
 #include "Widgets/SDMXPixelMappingPaletteItem.h"
-#include "ViewModels/DMXPixelMappingPaletteViewModel.h"
-#include "Templates/DMXPixelMappingComponentTemplate.h"
+
 #include "DragDrop/DMXPixelMappingDragDropOp.h"
+#include "Templates/DMXPixelMappingComponentTemplate.h"
+#include "ViewModels/DMXPixelMappingPaletteViewModel.h"
 
+#include "EditorStyleSet.h"
+#include "Widgets/Layout/SBorder.h"
+#include "Widgets/Text/STextBlock.h"
 #include "Widgets/Views/STableViewBase.h"
 #include "Widgets/Views/STableRow.h"
-#include "Widgets/Text/STextBlock.h"
-#include "EditorStyleSet.h"
 
 
+namespace UE::DMX::Private::SDMXPixelMappingPaletteItem
+{
+	constexpr FLinearColor NormalBGColor(0.f, 0.0f, 0.f, 0.f); // Fully transparent so it just shows the background
+	constexpr FLinearColor HoveredBGColor(1.f, 1.f, 1.f, 0.2f); // White transparent
+}
+
 void SDMXPixelMappingHierarchyItemHeader::Construct(const FArguments& InArgs, const TSharedRef< STableViewBase >& InOwnerTableView, const TSharedPtr<FDMXPixelMappingPaletteWidgetViewModel>& InViewModel)
 {
 	STableRow<FDMXPixelMappingPreviewWidgetViewModelPtr>::Construct(
@@ -27,6 +35,8 @@ void SDMXPixelMappingHierarchyItemHeader::Construct(const FArguments& InArgs, co
 
 void SDMXPixelMappingHierarchyItemTemplate::Construct(const FArguments& InArgs, const TSharedRef< STableViewBase >& InOwnerTableView, const TSharedPtr<FDMXPixelMappingPaletteWidgetViewModel>& InViewModel)
 {
+	using namespace UE::DMX::Private::SDMXPixelMappingPaletteItem;
+
 	ViewModel = InViewModel;
 
 	STableRow<FDMXPixelMappingPreviewWidgetViewModelPtr>::Construct(
@@ -37,8 +47,16 @@ void SDMXPixelMappingHierarchyItemTemplate::Construct(const FArguments& InArgs,
 		.OnDragDetected(this, &SDMXPixelMappingHierarchyItemTemplate::OnDraggingWidget)
 		.Content()
 		[
-			SNew(STextBlock)
-			.Text(InViewModel->GetName())
+			SNew(SBorder)
+			.BorderImage(FEditorStyle::GetBrush("WhiteBrush"))
+			.BorderBackgroundColor_Lambda([this]()
+				{
+					return IsHovered() ? HoveredBGColor : NormalBGColor;
+				})
+			[
+				SNew(STextBlock)
+				.Text(InViewModel->GetName())
+			]
 		],
 		InOwnerTableView);
 
diff --git a/Engine/Plugins/VirtualProduction/DMX/DMXPixelMapping/Source/DMXPixelMappingEditor/Private/Widgets/SDMXPixelMappingTransformHandle.cpp b/Engine/Plugins/VirtualProduction/DMX/DMXPixelMapping/Source/DMXPixelMappingEditor/Private/Widgets/SDMXPixelMappingTransformHandle.cpp
--- a/Engine/Plugins/VirtualProduction/DMX/DMXPixelMapping/Source/DMXPixelMappingEditor/Private/Widgets/SDMXPixelMappingTransformHandle.cpp
+++ b/Engine/Plugins/VirtualProduction/DMX/DMXPixelMapping/Source/DMXPixelMappingEditor/Private/Widgets/SDMXPixelMappingTransformHandle.cpp
@@ -46,27 +46,31 @@ FReply SDMXPixelMappingTransformHandle::OnMouseButtonDown(const FGeometry& MyGeo
 	{
 		Action = ComputeActionAtLocation(MyGeometry, MouseEvent);
 
-		check(DesignerViewWeakPtr.Pin());
-
-		const FDMXPixelMappingComponentReference& ComponentReference = DesignerViewWeakPtr.Pin()->GetSelectedComponent();
-		UDMXPixelMappingBaseComponent* Preview = ComponentReference.GetComponent();
-		UDMXPixelMappingBaseComponent* Component = ComponentReference.GetComponent();
-
-		if (UDMXPixelMappingOutputComponent* OutputComponent = Cast<UDMXPixelMappingOutputComponent>(Preview))
+		if (TSharedPtr<FDMXPixelMappingToolkit> Toolkit = DesignerViewWeakPtr.Pin()->GetToolkit())
 		{
-			FMargin Offsets;
-			FVector2D Size = OutputComponent->GetSize();
-			Offsets.Right = Size.X;
-			Offsets.Bottom = Size.Y;
-			StartingOffsets = Offsets;
-		}
+			const TSet<FDMXPixelMappingComponentReference> SelectedComponentReferences = Toolkit->GetSelectedComponents();
+			for (const FDMXPixelMappingComponentReference& ComponentReference : SelectedComponentReferences)
+			{
+				UDMXPixelMappingBaseComponent* Preview = ComponentReference.GetComponent();
+				UDMXPixelMappingBaseComponent* Component = ComponentReference.GetComponent();
 
-		MouseDownPosition = MouseEvent.GetScreenSpacePosition();
+				if (UDMXPixelMappingOutputComponent* OutputComponent = Cast<UDMXPixelMappingOutputComponent>(Preview))
+				{
+					FMargin Offsets;
+					FVector2D Size = OutputComponent->GetSize();
+					Offsets.Right = Size.X;
+					Offsets.Bottom = Size.Y;
+					StartingOffsets = Offsets;
+				}
 
-		ScopedTransaction = MakeShareable<FScopedTransaction>(new FScopedTransaction(LOCTEXT("ResizePixelMappingComponent", "PixelMapping: Resize Component")));
-		Component->Modify();
+				MouseDownPosition = MouseEvent.GetScreenSpacePosition();
 
-		return FReply::Handled().CaptureMouse(SharedThis(this));
+				ScopedTransaction = MakeShareable<FScopedTransaction>(new FScopedTransaction(LOCTEXT("ResizePixelMappingComponent", "PixelMapping: Resize Component")));
+				Component->Modify();
+
+				return FReply::Handled().CaptureMouse(SharedThis(this));
+			}
+		}
 	}
 
 	return FReply::Unhandled();
@@ -80,46 +84,46 @@ FReply SDMXPixelMappingTransformHandle::OnMouseButtonUp(const FGeometry& MyGeome
 		
 		if (TSharedPtr<FDMXPixelMappingToolkit> Toolkit = DesignerViewWeakPtr.Pin()->GetToolkit())
 		{		
-			const FDMXPixelMappingComponentReference& ComponentReference = DesignerViewWeakPtr.Pin()->GetSelectedComponent();
-			if (UDMXPixelMappingOutputComponent* ResizedComponent = Cast<UDMXPixelMappingOutputComponent>(ComponentReference.GetComponent()))
-			{			
-				// Delete the component if it's no longer over its parent
-				if (ResizedComponent->HasValidParent() &&
-					!ResizedComponent->IsOverParent())
+			const TSet<FDMXPixelMappingComponentReference> SelectedComponentReferences = Toolkit->GetSelectedComponents();
+			for (const FDMXPixelMappingComponentReference& ComponentReference : SelectedComponentReferences)
+			{
+				if (UDMXPixelMappingOutputComponent* ResizedComponent = Cast<UDMXPixelMappingOutputComponent>(ComponentReference.GetComponent()))
 				{
-					ResizedComponent->SetFlags(RF_Transactional);
-					ResizedComponent->Modify();
+					// Delete the component if it's no longer over its parent
+					if (ResizedComponent->HasValidParent() &&
+						!ResizedComponent->IsOverParent())
+					{
+						ResizedComponent->SetFlags(RF_Transactional);
+						ResizedComponent->Modify();
 
-					ResizedComponent->GetParent()->RemoveChild(ResizedComponent);
-				}
+						ResizedComponent->GetParent()->RemoveChild(ResizedComponent);
+					}
 
-				// Delete childs if they're no longer over their parent
-				for (UDMXPixelMappingBaseComponent* Child : TArray<UDMXPixelMappingBaseComponent*>(ResizedComponent->Children))
-				{
-					if (UDMXPixelMappingOutputComponent* OutputComponent = Cast<UDMXPixelMappingOutputComponent>(Child))
+					// Delete childs if they're no longer over their parent
+					for (UDMXPixelMappingBaseComponent* Child : TArray<UDMXPixelMappingBaseComponent*>(ResizedComponent->Children))
 					{
-						if (OutputComponent && !OutputComponent->IsOverParent())
+						if (UDMXPixelMappingOutputComponent* OutputComponent = Cast<UDMXPixelMappingOutputComponent>(Child))
 						{
-							ResizedComponent->RemoveChild(Child);
+							if (OutputComponent && !OutputComponent->IsOverParent())
+							{
+								ResizedComponent->RemoveChild(Child);
+							}
 						}
 					}
-				}
-			}
-		}
 
-		// Set the final size transacted
-		const FDMXPixelMappingComponentReference& ComponentReference = DesignerViewWeakPtr.Pin()->GetSelectedComponent();
-		UDMXPixelMappingBaseComponent* Component = ComponentReference.GetComponent();
+					// Set the final size transacted
+					const FVector2D Delta = MouseEvent.GetScreenSpacePosition() - MouseDownPosition;
+					const FVector2D TranslateAmount = Delta * (1.0f / (DesignerViewWeakPtr.Pin()->GetPreviewScale() * MyGeometry.Scale));
 
-		const FVector2D Delta = MouseEvent.GetScreenSpacePosition() - MouseDownPosition;
-		const FVector2D TranslateAmount = Delta * (1.0f / (DesignerViewWeakPtr.Pin()->GetPreviewScale() * MyGeometry.Scale));
-
-		Component->Modify();
-		Resize(Component, DragDirection, TranslateAmount);
+					ResizedComponent->Modify();
+					Resize(ResizedComponent, DragDirection, TranslateAmount);
 
-		ScopedTransaction.Reset();
+					ScopedTransaction.Reset();
 
-		return FReply::Handled().ReleaseMouseCapture();
+					return FReply::Handled().ReleaseMouseCapture();
+				}
+			}
+		}
 	}
 
 	return FReply::Unhandled();
@@ -129,15 +133,20 @@ FReply SDMXPixelMappingTransformHandle::OnMouseMove(const FGeometry& MyGeometry,
 {
 	if ( Action != EDMXPixelMappingTransformAction::None )
 	{
-		check(DesignerViewWeakPtr.Pin());
-
-		const FDMXPixelMappingComponentReference& ComponentReference = DesignerViewWeakPtr.Pin()->GetSelectedComponent();
-		UDMXPixelMappingBaseComponent* Component = ComponentReference.GetComponent();
-
-		const FVector2D Delta = MouseEvent.GetScreenSpacePosition() - MouseDownPosition;
-		const FVector2D TranslateAmount = Delta * (1.0f / (DesignerViewWeakPtr.Pin()->GetPreviewScale() * MyGeometry.Scale));
+		if (TSharedPtr<FDMXPixelMappingToolkit> Toolkit = DesignerViewWeakPtr.Pin()->GetToolkit())
+		{
+			const TSet<FDMXPixelMappingComponentReference> SelectedComponentReferences = Toolkit->GetSelectedComponents();
+			for (const FDMXPixelMappingComponentReference& ComponentReference : SelectedComponentReferences)
+			{
+				if (UDMXPixelMappingOutputComponent* ResizedComponent = Cast<UDMXPixelMappingOutputComponent>(ComponentReference.GetComponent()))
+				{
+					const FVector2D Delta = MouseEvent.GetScreenSpacePosition() - MouseDownPosition;
+					const FVector2D TranslateAmount = Delta * (1.0f / (DesignerViewWeakPtr.Pin()->GetPreviewScale() * MyGeometry.Scale));
 
-		Resize(Component, DragDirection, TranslateAmount);
+					Resize(ResizedComponent, DragDirection, TranslateAmount);
+				}
+			}
+		}
 	}
 
 	return FReply::Unhandled();
diff --git a/Engine/Plugins/VirtualProduction/DMX/DMXPixelMapping/Source/DMXPixelMappingEditorWidgets/Private/DMXPixelMappingComponentWidget.cpp b/Engine/Plugins/VirtualProduction/DMX/DMXPixelMapping/Source/DMXPixelMappingEditorWidgets/Private/DMXPixelMappingComponentWidget.cpp
--- a/Engine/Plugins/VirtualProduction/DMX/DMXPixelMapping/Source/DMXPixelMappingEditorWidgets/Private/DMXPixelMappingComponentWidget.cpp
+++ b/Engine/Plugins/VirtualProduction/DMX/DMXPixelMapping/Source/DMXPixelMappingEditorWidgets/Private/DMXPixelMappingComponentWidget.cpp
@@ -52,16 +52,22 @@ void FDMXPixelMappingComponentWidget::AddToCanvas(const TSharedRef<SConstraintCa
 	];
 }
 
+FDMXPixelMappingComponentWidget::~FDMXPixelMappingComponentWidget()
+{
+	RemoveFromCanvas();
+}
+
 void FDMXPixelMappingComponentWidget::RemoveFromCanvas()
 {
 	if (OuterCanvas.IsValid() && ComponentSlot && LabelSlot)
 	{
 		OuterCanvas->RemoveSlot(ComponentSlot->GetWidget());
 		OuterCanvas->RemoveSlot(LabelSlot->GetWidget());
-
-		ComponentSlot = nullptr;
-		LabelSlot = nullptr;
 	}
+
+	ComponentSlot = nullptr;
+	LabelSlot = nullptr;
+	OuterCanvas.Reset();
 }
 
 void FDMXPixelMappingComponentWidget::SetZOrder(float ZOrder)
diff --git a/Engine/Plugins/VirtualProduction/DMX/DMXPixelMapping/Source/DMXPixelMappingEditorWidgets/Private/SDMXPixelMappingComponentBox.cpp b/Engine/Plugins/VirtualProduction/DMX/DMXPixelMapping/Source/DMXPixelMappingEditorWidgets/Private/SDMXPixelMappingComponentBox.cpp
--- a/Engine/Plugins/VirtualProduction/DMX/DMXPixelMapping/Source/DMXPixelMappingEditorWidgets/Private/SDMXPixelMappingComponentBox.cpp
+++ b/Engine/Plugins/VirtualProduction/DMX/DMXPixelMapping/Source/DMXPixelMappingEditorWidgets/Private/SDMXPixelMappingComponentBox.cpp
@@ -10,6 +10,8 @@
 
 void SDMXPixelMappingComponentBox::Construct(const FArguments& InArgs)
 {
+	SetCanTick(false);
+
 	IDText = InArgs._IDText;
 	MaxIDTextSize = InArgs._MaxIDTextSize;
 
diff --git a/Engine/Plugins/VirtualProduction/DMX/DMXPixelMapping/Source/DMXPixelMappingEditorWidgets/Private/SDMXPixelMappingComponentLabel.cpp b/Engine/Plugins/VirtualProduction/DMX/DMXPixelMapping/Source/DMXPixelMappingEditorWidgets/Private/SDMXPixelMappingComponentLabel.cpp
--- a/Engine/Plugins/VirtualProduction/DMX/DMXPixelMapping/Source/DMXPixelMappingEditorWidgets/Private/SDMXPixelMappingComponentLabel.cpp
+++ b/Engine/Plugins/VirtualProduction/DMX/DMXPixelMapping/Source/DMXPixelMappingEditorWidgets/Private/SDMXPixelMappingComponentLabel.cpp
@@ -12,6 +12,8 @@
 
 void SDMXPixelMappingComponentLabel::Construct(const FArguments& InArgs)
 {
+	SetCanTick(false);
+
 	LabelText = InArgs._LabelText;
 	bAlignAbove = InArgs._bAlignAbove;
 
diff --git a/Engine/Plugins/VirtualProduction/DMX/DMXPixelMapping/Source/DMXPixelMappingEditorWidgets/Public/DMXPixelMappingComponentWidget.h b/Engine/Plugins/VirtualProduction/DMX/DMXPixelMapping/Source/DMXPixelMappingEditorWidgets/Public/DMXPixelMappingComponentWidget.h
--- a/Engine/Plugins/VirtualProduction/DMX/DMXPixelMapping/Source/DMXPixelMappingEditorWidgets/Public/DMXPixelMappingComponentWidget.h
+++ b/Engine/Plugins/VirtualProduction/DMX/DMXPixelMapping/Source/DMXPixelMappingEditorWidgets/Public/DMXPixelMappingComponentWidget.h
@@ -22,6 +22,9 @@ public:
 	/** Constructs the class, optionally with custom widgets */
 	FDMXPixelMappingComponentWidget(TSharedPtr<SDMXPixelMappingComponentBox> InComponentBox = nullptr, TSharedPtr<SDMXPixelMappingComponentLabel> InComponentLabel = nullptr, bool bShowLabelAbove = false);
 
+	/** Destructor */
+	virtual ~FDMXPixelMappingComponentWidget();
+
 	/** Adds the widgets to a canvas. If it already resides in a canvas it is removed from that. */
 	void AddToCanvas(const TSharedRef<SConstraintCanvas>& Canvas, float ZOrder);
 
@@ -43,6 +46,9 @@ public:
 	/** Sets the label text */
 	void SetLabelText(const FText& Text);
 
+	/** Sets if the ID text of is visible */
+	void SetIDVisibility(bool bVisible);
+
 	/** Sets the color of the widget */
 	void SetColor(const FLinearColor& Color);
 
diff --git a/Engine/Plugins/VirtualProduction/DMX/DMXPixelMapping/Source/DMXPixelMappingRuntime/Private/Components/DMXPixelMappingBaseComponent.cpp b/Engine/Plugins/VirtualProduction/DMX/DMXPixelMapping/Source/DMXPixelMappingRuntime/Private/Components/DMXPixelMappingBaseComponent.cpp
--- a/Engine/Plugins/VirtualProduction/DMX/DMXPixelMapping/Source/DMXPixelMappingRuntime/Private/Components/DMXPixelMappingBaseComponent.cpp
+++ b/Engine/Plugins/VirtualProduction/DMX/DMXPixelMapping/Source/DMXPixelMappingRuntime/Private/Components/DMXPixelMappingBaseComponent.cpp
@@ -184,7 +184,8 @@ void UDMXPixelMappingBaseComponent::AddChild(UDMXPixelMappingBaseComponent* InCo
 		InComponent->WeakParent = this;
 
 		// Allow children to be readded, this may be the case during Undo of a Remove
-			Children.AddUnique(InComponent);
+		Children.AddUnique(InComponent);
+		InComponent->NotifyAddedToParent();
 
 		// Broadcast the change
 		GetOnComponentAdded().Broadcast(GetPixelMapping(), InComponent);
@@ -193,19 +194,12 @@ void UDMXPixelMappingBaseComponent::AddChild(UDMXPixelMappingBaseComponent* InCo
 
 void UDMXPixelMappingBaseComponent::RemoveChild(UDMXPixelMappingBaseComponent* ChildComponent)
 {
-	if (ensureMsgf(ChildComponent || Children.Contains(ChildComponent), TEXT("Trying to remove child, but %s is not a child of %s."), *ChildComponent->GetUserFriendlyName(), *GetUserFriendlyName()))
+	if (ChildComponent)
 	{
-		TArray<UDMXPixelMappingBaseComponent*> ChildComponents;
-		ChildComponent->GetChildComponentsRecursively(ChildComponents);
-		for (UDMXPixelMappingBaseComponent* ChildOfChild : ChildComponents)
-		{
-			// Recursively call this function here on children
-			ChildComponent->RemoveChild(ChildOfChild);
-		}
-
 		ChildComponent->ResetDMX();
 
 		Children.Remove(ChildComponent);
+		ChildComponent->NotifyRemovedFromParent();
 
 		ChildComponent->SetFlags(RF_Transactional);
 
@@ -235,3 +229,18 @@ void UDMXPixelMappingBaseComponent::GetChildComponentsRecursively(TArray<UDMXPix
 		Components.Add(InComponent);
 	});
 }
+
+void UDMXPixelMappingBaseComponent::NotifyAddedToParent()
+{
+	// Nothing in base
+}
+
+void UDMXPixelMappingBaseComponent::NotifyRemovedFromParent()
+{
+	for (UDMXPixelMappingBaseComponent* Child : TArray< UDMXPixelMappingBaseComponent*>(Children))
+	{
+		RemoveChild(Child);
+	}
+
+	Children.Reset();
+}
diff --git a/Engine/Plugins/VirtualProduction/DMX/DMXPixelMapping/Source/DMXPixelMappingRuntime/Private/Components/DMXPixelMappingMatrixCellComponent.cpp b/Engine/Plugins/VirtualProduction/DMX/DMXPixelMapping/Source/DMXPixelMappingRuntime/Private/Components/DMXPixelMappingMatrixCellComponent.cpp
--- a/Engine/Plugins/VirtualProduction/DMX/DMXPixelMapping/Source/DMXPixelMappingRuntime/Private/Components/DMXPixelMappingMatrixCellComponent.cpp
+++ b/Engine/Plugins/VirtualProduction/DMX/DMXPixelMapping/Source/DMXPixelMappingRuntime/Private/Components/DMXPixelMappingMatrixCellComponent.cpp
@@ -69,24 +69,6 @@ void UDMXPixelMappingMatrixCellComponent::PostEditChangeProperty(FPropertyChange
 }
 #endif // WITH_EDITOR
 
-#if WITH_EDITOR
-void UDMXPixelMappingMatrixCellComponent::PreEditUndo()
-{
-	// Use default engine instead of parent class.
-	// Let the prent matrix component handle it instead.
-	UObject::PreEditUndo();
-}
-#endif // WITH_EDITOR
-
-#if WITH_EDITOR
-void UDMXPixelMappingMatrixCellComponent::PostEditUndo()
-	{
-	// Use default engine instead of parent class.
-	// Let the prent matrix component handle it instead.	
-	UObject::PostEditUndo();
-	}
-#endif // WITH_EDITOR
-
 #if WITH_EDITOR
 TSharedRef<FDMXPixelMappingComponentWidget> UDMXPixelMappingMatrixCellComponent::BuildSlot(TSharedRef<SConstraintCanvas> InCanvas)
 {
diff --git a/Engine/Plugins/VirtualProduction/DMX/DMXPixelMapping/Source/DMXPixelMappingRuntime/Private/Components/DMXPixelMappingMatrixComponent.cpp b/Engine/Plugins/VirtualProduction/DMX/DMXPixelMapping/Source/DMXPixelMappingRuntime/Private/Components/DMXPixelMappingMatrixComponent.cpp
--- a/Engine/Plugins/VirtualProduction/DMX/DMXPixelMapping/Source/DMXPixelMappingRuntime/Private/Components/DMXPixelMappingMatrixComponent.cpp
+++ b/Engine/Plugins/VirtualProduction/DMX/DMXPixelMapping/Source/DMXPixelMappingRuntime/Private/Components/DMXPixelMappingMatrixComponent.cpp
@@ -161,41 +161,6 @@ void UDMXPixelMappingMatrixComponent::PostEditChangeChainProperty(FPropertyChang
 }
 #endif // WITH_EDITOR
 
-#if WITH_EDITOR
-void UDMXPixelMappingMatrixComponent::PreEditUndo()
-{
-	// Override Output Component, but not UObject
-	UObject::PreEditUndo();
-
-	PreEditUndoMatrixCellChildren = Children;
-}
-#endif // WITH_EDITOR
-
-#if WITH_EDITOR
-void UDMXPixelMappingMatrixComponent::PostEditUndo()
-{
-	// Override Output Component, but not UObject
-	UObject::PostEditUndo();
-
-	HandleMatrixChanged();
-	HandlePositionChanged();
-
-	const EVisibility NewVisibility = IsVisible() ? EVisibility::Visible : EVisibility::Collapsed;
-
-	constexpr bool bPropagonateToChildren = true;
-	UpdateComponentWidget(NewVisibility, bPropagonateToChildren);
-
-	// HandleSizeOrMatrixChanged creates new cells. Hide the previous ones
-	for (UDMXPixelMappingBaseComponent* PreEditUndoChild : PreEditUndoMatrixCellChildren)
-	{
-		// Children are always matrix cells
-		UDMXPixelMappingMatrixCellComponent* PreEditUndoMatrixCellComponent = Cast<UDMXPixelMappingMatrixCellComponent>(PreEditUndoChild);
-
-		PreEditUndoMatrixCellComponent->UpdateComponentWidget(EVisibility::Collapsed);
-	}
-}
-#endif // WITH_EDITOR
-
 #if WITH_EDITOR
 const FText UDMXPixelMappingMatrixComponent::GetPaletteCategory()
 {
@@ -357,7 +322,7 @@ void UDMXPixelMappingMatrixComponent::SetSize(const FVector2D& NewSize)
 	SizeX = FMath::Max(SizeX, 1.f);
 	SizeY = FMath::Max(SizeY, 1.f);
 
-	HandleMatrixChanged();
+	//HandleMatrixChanged();
 	HandleSizeChanged();
 }
 
diff --git a/Engine/Plugins/VirtualProduction/DMX/DMXPixelMapping/Source/DMXPixelMappingRuntime/Private/Components/DMXPixelMappingOutputComponent.cpp b/Engine/Plugins/VirtualProduction/DMX/DMXPixelMapping/Source/DMXPixelMappingRuntime/Private/Components/DMXPixelMappingOutputComponent.cpp
--- a/Engine/Plugins/VirtualProduction/DMX/DMXPixelMapping/Source/DMXPixelMappingRuntime/Private/Components/DMXPixelMappingOutputComponent.cpp
+++ b/Engine/Plugins/VirtualProduction/DMX/DMXPixelMapping/Source/DMXPixelMappingRuntime/Private/Components/DMXPixelMappingOutputComponent.cpp
@@ -40,12 +40,12 @@ void UDMXPixelMappingOutputComponent::PostEditChangeProperty(FPropertyChangedEve
 	
 	if (PropertyChangedEvent.GetPropertyName() == GET_MEMBER_NAME_CHECKED(UDMXPixelMappingOutputComponent, bVisibleInDesigner))
 	{
-		const EVisibility NewVisiblity = IsVisible() ? EVisibility::Visible : EVisibility::Collapsed;
+		const EVisibility NewVisiblity = IsVisible() ? EVisibility::HitTestInvisible : EVisibility::Hidden;
 
-				if (ComponentWidget.IsValid())
-				{
-					ComponentWidget->SetVisibility(NewVisiblity);
-				}
+		if (ComponentWidget.IsValid())
+		{
+			ComponentWidget->SetVisibility(NewVisiblity);
+		}
 
 		constexpr bool bSetVisibilityRecursive = true;
 		ForEachChildOfClass<UDMXPixelMappingOutputComponent>([NewVisiblity](UDMXPixelMappingOutputComponent* ChildComponent)
@@ -121,7 +121,7 @@ void UDMXPixelMappingOutputComponent::PostEditUndo()
 {
 	Super::PostEditUndo();
 
-	EVisibility NewVisibility = IsVisible() ? EVisibility::Visible : EVisibility::Collapsed;
+	EVisibility NewVisibility = IsVisible() ? EVisibility::HitTestInvisible : EVisibility::Hidden;
 
 	UpdateComponentWidget(NewVisibility);
 
@@ -132,7 +132,7 @@ void UDMXPixelMappingOutputComponent::PostEditUndo()
 		{
 			if (UDMXPixelMappingOutputComponent* PreEditUndoOutputComponent = Cast<UDMXPixelMappingOutputComponent>(PreEditUndoChild))
 			{
-				const EVisibility NewChildVisibility = PreEditUndoOutputComponent->IsVisible() ? EVisibility::Visible : EVisibility::Collapsed;
+				const EVisibility NewChildVisibility = PreEditUndoOutputComponent->IsVisible() ? EVisibility::HitTestInvisible : EVisibility::Hidden;
 
 				constexpr bool bWithChildrenRecursive = false;
 				PreEditUndoOutputComponent->UpdateComponentWidget(NewChildVisibility, bWithChildrenRecursive);
@@ -146,7 +146,7 @@ void UDMXPixelMappingOutputComponent::PostEditUndo()
 		{
 			if (UDMXPixelMappingOutputComponent* PreEditUndoOutputComponent = Cast<UDMXPixelMappingOutputComponent>(PreEditUndoChild))
 			{
-				const EVisibility NewChildVisibility = IsVisible() ? EVisibility::Visible : EVisibility::Collapsed;
+				const EVisibility NewChildVisibility = IsVisible() ? EVisibility::HitTestInvisible : EVisibility::Hidden;
 
 				constexpr bool bWithChildrenRecursive = false;
 				PreEditUndoOutputComponent->UpdateComponentWidget(NewChildVisibility, bWithChildrenRecursive);
@@ -178,28 +178,28 @@ bool UDMXPixelMappingOutputComponent::CanBeMovedTo(const UDMXPixelMappingBaseCom
 	return false;
 }
 
-void UDMXPixelMappingOutputComponent::AddChild(UDMXPixelMappingBaseComponent* InComponent)
+void UDMXPixelMappingOutputComponent::NotifyAddedToParent()
 {
-	Super::AddChild(InComponent);
+	Super::NotifyAddedToParent();
 
 #if WITH_EDITOR
-	if (UDMXPixelMappingOutputComponent* ChildOutputComponent = Cast<UDMXPixelMappingOutputComponent>(InComponent))
-	{
-		const EVisibility NewChildVisibility = ChildOutputComponent->IsVisible() ? EVisibility::Visible : EVisibility::Collapsed;
-
-		ChildOutputComponent->UpdateComponentWidget(NewChildVisibility);
-	}	
+	UpdateComponentWidget(EVisibility::HitTestInvisible);
 #endif // WITH_EDITOR
 }
 
-void UDMXPixelMappingOutputComponent::RemoveChild(UDMXPixelMappingBaseComponent* InComponent)
+void UDMXPixelMappingOutputComponent::NotifyRemovedFromParent()
 {
-	Super::RemoveChild(InComponent);
+	Super::NotifyRemovedFromParent();
 
 #if WITH_EDITOR
-	if (UDMXPixelMappingOutputComponent* ChildOutputComponent = Cast<UDMXPixelMappingOutputComponent>(InComponent))
+	if (ComponentWidget.IsValid())
 	{
-		ChildOutputComponent->UpdateComponentWidget(EVisibility::Collapsed);
+		ComponentWidget->RemoveFromCanvas();
+
+		// Should have released all references by now
+		ensureMsgf(ComponentWidget.GetSharedReferenceCount() == 1, TEXT("Detected Reference to Component Widget the component is destroyed."));
+
+		ComponentWidget.Reset();
 	}
 #endif // WITH_EDITOR
 }
@@ -216,13 +216,10 @@ const FText UDMXPixelMappingOutputComponent::GetPaletteCategory()
 #if WITH_EDITOR
 TSharedRef<FDMXPixelMappingComponentWidget> UDMXPixelMappingOutputComponent::BuildSlot(TSharedRef<SConstraintCanvas> InCanvas)
 {
-	if (!ComponentWidget.IsValid())
-	{
-		ComponentWidget = MakeShared<FDMXPixelMappingComponentWidget>();
-		ComponentWidget->AddToCanvas(InCanvas, ZOrder);
-	}
+	ComponentWidget = MakeShared<FDMXPixelMappingComponentWidget>();
+	ComponentWidget->AddToCanvas(InCanvas, ZOrder);
 
-	EVisibility NewVisibility = IsVisible() ? EVisibility::Visible : EVisibility::Collapsed;
+	EVisibility NewVisibility = IsVisible() ? EVisibility::HitTestInvisible : EVisibility::Hidden;
 
 	UpdateComponentWidget(NewVisibility);
 
@@ -247,12 +244,9 @@ void UDMXPixelMappingOutputComponent::SetZOrder(int32 NewZOrder)
 	ZOrder = NewZOrder;
 
 	// Apply to the UI
-	if (TSharedPtr<SConstraintCanvas> Canvas = FindRendererComponentCanvas())
+	if (ComponentWidget.IsValid())
 	{
-		if (ComponentWidget.IsValid())
-		{
-			ComponentWidget->SetZOrder(ZOrder);
-		}
+		ComponentWidget->SetZOrder(ZOrder);
 	}
 
 	constexpr bool bRecursive = true;
@@ -372,19 +366,6 @@ void UDMXPixelMappingOutputComponent::MakeHighestZOrderInComponentRect()
 }
 #endif // WITH_EDITOR
 
-#if WITH_EDITOR
-TSharedPtr<SConstraintCanvas> UDMXPixelMappingOutputComponent::FindRendererComponentCanvas() const
-{
-	UDMXPixelMappingRendererComponent* ParentRendererComponent = FindRendererComponent();
-	if(ParentRendererComponent && ParentRendererComponent->GetComponentsCanvas().IsValid())
-	{
-		return ParentRendererComponent->GetComponentsCanvas();
-	}
-
-	return nullptr;
-}
-#endif // WITH_EDITOR
-
 #if WITH_EDITOR
 void UDMXPixelMappingOutputComponent::UpdateComponentWidget(EVisibility NewVisibility, bool bWithChildrenRecursive)
 {
@@ -407,7 +388,7 @@ void UDMXPixelMappingOutputComponent::UpdateComponentWidget(EVisibility NewVisib
 				ChildOutputComponent->UpdateComponentWidget(NewVisibility, bWithChildrenRecursive);
 			}
 		}
-}
+	}
 }
 #endif // WITH_EDITOR
 
diff --git a/Engine/Plugins/VirtualProduction/DMX/DMXPixelMapping/Source/DMXPixelMappingRuntime/Private/Components/DMXPixelMappingRendererComponent.cpp b/Engine/Plugins/VirtualProduction/DMX/DMXPixelMapping/Source/DMXPixelMappingRuntime/Private/Components/DMXPixelMappingRendererComponent.cpp
--- a/Engine/Plugins/VirtualProduction/DMX/DMXPixelMapping/Source/DMXPixelMappingRuntime/Private/Components/DMXPixelMappingRendererComponent.cpp
+++ b/Engine/Plugins/VirtualProduction/DMX/DMXPixelMapping/Source/DMXPixelMappingRuntime/Private/Components/DMXPixelMappingRendererComponent.cpp
@@ -5,8 +5,11 @@
 #include "DMXPixelMappingRuntimeObjectVersion.h"
 #include "DMXPixelMappingTypes.h"
 #include "IDMXPixelMappingRendererModule.h"
+#include "Components/DMXPixelMappingFixtureGroupComponent.h"
+#include "Components/DMXPixelMappingFixtureGroupItemComponent.h"
 #include "Components/DMXPixelMappingOutputComponent.h"
 #include "Components/DMXPixelMappingOutputDMXComponent.h"
+#include "Components/DMXPixelMappingMatrixComponent.h"
 #include "Components/DMXPixelMappingRootComponent.h"
 #include "Components/DMXPixelMappingScreenComponent.h"
 #include "Library/DMXEntityFixtureType.h"
@@ -180,7 +183,7 @@ void UDMXPixelMappingRendererComponent::PostEditChangeChainProperty(FPropertyCha
 #if WITH_EDITOR
 void UDMXPixelMappingRendererComponent::RenderEditorPreviewTexture()
 {
-	if (DownsampleBufferTarget == nullptr || DownsamplePixelCount == 0)
+	if (!DownsampleBufferTarget)
 	{
 		return;
 	}
@@ -239,6 +242,37 @@ UTextureRenderTarget2D* UDMXPixelMappingRendererComponent::GetPreviewRenderTarge
 }
 #endif // WITH_EDITOR
 
+bool UDMXPixelMappingRendererComponent::GetPixelMappingComponentModulators(FDMXEntityFixturePatchRef FixturePatchRef, TArray<UDMXModulator*>& DMXModulators)
+{
+	for (const UDMXPixelMappingBaseComponent* Child : Children)
+	{
+		if (const UDMXPixelMappingFixtureGroupComponent* GroupComponent = Cast<UDMXPixelMappingFixtureGroupComponent>(Child))
+		{
+			for (const UDMXPixelMappingBaseComponent* ChildOfGroupComponent : GroupComponent->Children)
+			{
+				if (const UDMXPixelMappingFixtureGroupItemComponent* GroupItemComponent = Cast<UDMXPixelMappingFixtureGroupItemComponent>(ChildOfGroupComponent))
+				{
+					if (GroupItemComponent->FixturePatchRef.GetFixturePatch() == FixturePatchRef.GetFixturePatch())
+					{
+						DMXModulators = GroupItemComponent->Modulators;
+						return true;
+					}
+				}
+				else if (const UDMXPixelMappingMatrixComponent* MatrixComponent = Cast<UDMXPixelMappingMatrixComponent>(ChildOfGroupComponent))
+				{
+					if (MatrixComponent->FixturePatchRef.GetFixturePatch() == FixturePatchRef.GetFixturePatch())
+					{
+						DMXModulators = MatrixComponent->Modulators;
+						return true;
+					}
+				}
+			}
+		}
+	}
+
+	return false;
+}
+
 #if WITH_EDITOR
 TSharedRef<SWidget> UDMXPixelMappingRendererComponent::TakeWidget()
 {
diff --git a/Engine/Plugins/VirtualProduction/DMX/DMXPixelMapping/Source/DMXPixelMappingRuntime/Public/Components/DMXPixelMappingBaseComponent.h b/Engine/Plugins/VirtualProduction/DMX/DMXPixelMapping/Source/DMXPixelMappingRuntime/Public/Components/DMXPixelMappingBaseComponent.h
--- a/Engine/Plugins/VirtualProduction/DMX/DMXPixelMapping/Source/DMXPixelMappingRuntime/Public/Components/DMXPixelMappingBaseComponent.h
+++ b/Engine/Plugins/VirtualProduction/DMX/DMXPixelMapping/Source/DMXPixelMappingRuntime/Public/Components/DMXPixelMappingBaseComponent.h
@@ -264,6 +264,13 @@ public:
 	UPROPERTY(Meta = (DeprecatedProperty, DeprecationMessage = "Leads to entangled references. Use GetParent() or WeakParent instead."))
 	UDMXPixelMappingBaseComponent* Parent_DEPRECATED;
 
+protected:
+	/** Called when the component was added to a parent */
+	virtual void NotifyAddedToParent();
+
+	/** Called when the component was added to a parent */
+	virtual void NotifyRemovedFromParent();
+
 private:
 	/** Parent component */
 	UPROPERTY(NonTransactional)
diff --git a/Engine/Plugins/VirtualProduction/DMX/DMXPixelMapping/Source/DMXPixelMappingRuntime/Public/Components/DMXPixelMappingMatrixCellComponent.h b/Engine/Plugins/VirtualProduction/DMX/DMXPixelMapping/Source/DMXPixelMappingRuntime/Public/Components/DMXPixelMappingMatrixCellComponent.h
--- a/Engine/Plugins/VirtualProduction/DMX/DMXPixelMapping/Source/DMXPixelMappingRuntime/Public/Components/DMXPixelMappingMatrixCellComponent.h
+++ b/Engine/Plugins/VirtualProduction/DMX/DMXPixelMapping/Source/DMXPixelMappingRuntime/Public/Components/DMXPixelMappingMatrixCellComponent.h
@@ -32,8 +32,6 @@ public:
 	//~ Begin UObject implementation
 #if WITH_EDITOR
 	virtual void PostEditChangeProperty(FPropertyChangedEvent& PropertyChangedEvent) override;
-	virtual void PreEditUndo() override;
-	virtual void PostEditUndo() override;
 #endif // WITH_EDITOR
 	//~ End UObject implementation
 
diff --git a/Engine/Plugins/VirtualProduction/DMX/DMXPixelMapping/Source/DMXPixelMappingRuntime/Public/Components/DMXPixelMappingMatrixComponent.h b/Engine/Plugins/VirtualProduction/DMX/DMXPixelMapping/Source/DMXPixelMappingRuntime/Public/Components/DMXPixelMappingMatrixComponent.h
--- a/Engine/Plugins/VirtualProduction/DMX/DMXPixelMapping/Source/DMXPixelMappingRuntime/Public/Components/DMXPixelMappingMatrixComponent.h
+++ b/Engine/Plugins/VirtualProduction/DMX/DMXPixelMapping/Source/DMXPixelMappingRuntime/Public/Components/DMXPixelMappingMatrixComponent.h
@@ -43,8 +43,6 @@ public:
 #if WITH_EDITOR
 	virtual void PostEditChangeProperty(FPropertyChangedEvent& PropertyChangedEvent) override;
 	virtual void PostEditChangeChainProperty(FPropertyChangedChainEvent& PropertyChangedChainEvent) override;
-	virtual void PreEditUndo() override;
-	virtual void PostEditUndo() override; 
 #endif // WITH_EDITOR
 	// ~End UObject interface
 	
diff --git a/Engine/Plugins/VirtualProduction/DMX/DMXPixelMapping/Source/DMXPixelMappingRuntime/Public/Components/DMXPixelMappingOutputComponent.h b/Engine/Plugins/VirtualProduction/DMX/DMXPixelMapping/Source/DMXPixelMappingRuntime/Public/Components/DMXPixelMappingOutputComponent.h
--- a/Engine/Plugins/VirtualProduction/DMX/DMXPixelMapping/Source/DMXPixelMappingRuntime/Public/Components/DMXPixelMappingOutputComponent.h
+++ b/Engine/Plugins/VirtualProduction/DMX/DMXPixelMapping/Source/DMXPixelMappingRuntime/Public/Components/DMXPixelMappingOutputComponent.h
@@ -72,11 +72,10 @@ protected:
 public:
 	//~ Begin DMXPixelMappingBaseComponent interface
 	virtual bool CanBeMovedTo(const UDMXPixelMappingBaseComponent* Component) const override;
-	virtual void AddChild(UDMXPixelMappingBaseComponent* InComponent) override;
-	virtual void RemoveChild(UDMXPixelMappingBaseComponent* InComponent) override;
+	virtual void NotifyAddedToParent() override;
+	virtual void NotifyRemovedFromParent() override;
 	//~ End DMXPixelMappingBaseComponent interface
  
-
 	/*----------------------------------
 		UDMXPixelMappingOutputComponent Interface
 	----------------------------------*/
@@ -148,9 +147,6 @@ public:
 	FORCEINLINE const TSharedPtr<FDMXPixelMappingComponentWidget> GetComponentWidget() { return ComponentWidget; }
 
 protected:
-	/** Returns the canvas of the render component if available */
-	TSharedPtr<SConstraintCanvas> FindRendererComponentCanvas() const;
-
 	/** Udpates the component widget. If bWithChildrenRecursive, updates child Components' Component Widget recursively */
 	void UpdateComponentWidget(EVisibility NewVisibility = EVisibility::Visible, bool bWithChildrenRecursive = false);
 #endif // WITH_EDITOR
diff --git a/Engine/Plugins/VirtualProduction/DMX/DMXPixelMapping/Source/DMXPixelMappingRuntime/Public/Components/DMXPixelMappingRendererComponent.h b/Engine/Plugins/VirtualProduction/DMX/DMXPixelMapping/Source/DMXPixelMappingRuntime/Public/Components/DMXPixelMappingRendererComponent.h
--- a/Engine/Plugins/VirtualProduction/DMX/DMXPixelMapping/Source/DMXPixelMappingRuntime/Public/Components/DMXPixelMappingRendererComponent.h
+++ b/Engine/Plugins/VirtualProduction/DMX/DMXPixelMapping/Source/DMXPixelMappingRuntime/Public/Components/DMXPixelMappingRendererComponent.h
@@ -3,8 +3,12 @@
 #pragma once
 
 #include "Components/DMXPixelMappingOutputComponent.h"
+
 #include "IDMXPixelMappingRenderer.h"
+#include "Library/DMXEntityReference.h"
+
 #include "Templates/SubclassOf.h"
+
 #include "DMXPixelMappingRendererComponent.generated.h"
 
 
@@ -153,6 +157,13 @@ private:
 	UTextureRenderTarget2D* CreateRenderTarget(const FName& InBaseName);
 
 public:
+	/**
+	 * Returns the Modulators of the component corresponding to the patch specified. 
+	 * Note, this node does a lookup on all fixture patches in use, hence may be slow and shouldn't be called on tick. 
+	 */
+	UFUNCTION(BlueprintCallable, Category = "DMX")
+	bool GetPixelMappingComponentModulators(FDMXEntityFixturePatchRef FixturePatchRef, TArray<UDMXModulator*>& DMXModulators);
+
 	/** Type of rendering, Texture, Material, UMG, etc... */
 	UPROPERTY(EditAnywhere, BlueprintReadOnly, Category = "Render Settings")
 	EDMXPixelMappingRendererType RendererType;
@@ -234,4 +245,4 @@ private:
 public:
 	/** Max downsample target size */
 	static const FIntPoint MaxDownsampleBufferTargetSize;
-};
\ No newline at end of file
+};
