DMX - PixellMapping smaller default size for Matrix Components, larger for Group Items, and both the same

diff --git a/Engine/Plugins/VirtualProduction/DMX/DMXPixelMapping/Source/DMXPixelMappingRuntime/Private/Components/DMXPixelMappingFixtureGroupItemComponent.cpp b/Engine/Plugins/VirtualProduction/DMX/DMXPixelMapping/Source/DMXPixelMappingRuntime/Private/Components/DMXPixelMappingFixtureGroupItemComponent.cpp
--- a/Engine/Plugins/VirtualProduction/DMX/DMXPixelMapping/Source/DMXPixelMappingRuntime/Private/Components/DMXPixelMappingFixtureGroupItemComponent.cpp
+++ b/Engine/Plugins/VirtualProduction/DMX/DMXPixelMapping/Source/DMXPixelMappingRuntime/Private/Components/DMXPixelMappingFixtureGroupItemComponent.cpp
@@ -28,8 +28,8 @@ DECLARE_CYCLE_STAT(TEXT("Send Fixture Group Item"), STAT_DMXPixelMaping_FixtureG
 UDMXPixelMappingFixtureGroupItemComponent::UDMXPixelMappingFixtureGroupItemComponent()
 	: DownsamplePixelIndex(0)
 {
-	SizeX = 25.f;
-	SizeY = 25.f;
+	SizeX = 32.f;
+	SizeY = 32.f;
 
 	ColorMode = EDMXColorMode::CM_RGB;
 	AttributeRExpose = AttributeGExpose = AttributeBExpose = true;
@@ -44,7 +44,6 @@ UDMXPixelMappingFixtureGroupItemComponent::UDMXPixelMappingFixtureGroupItemCompo
 #endif // WITH_EDITOR
 }
 
-
 void UDMXPixelMappingFixtureGroupItemComponent::PostLoad()
 {
 	Super::PostLoad();
diff --git a/Engine/Plugins/VirtualProduction/DMX/DMXPixelMapping/Source/DMXPixelMappingRuntime/Private/Components/DMXPixelMappingMatrixComponent.cpp b/Engine/Plugins/VirtualProduction/DMX/DMXPixelMapping/Source/DMXPixelMappingRuntime/Private/Components/DMXPixelMappingMatrixComponent.cpp
--- a/Engine/Plugins/VirtualProduction/DMX/DMXPixelMapping/Source/DMXPixelMappingRuntime/Private/Components/DMXPixelMappingMatrixComponent.cpp
+++ b/Engine/Plugins/VirtualProduction/DMX/DMXPixelMapping/Source/DMXPixelMappingRuntime/Private/Components/DMXPixelMappingMatrixComponent.cpp
@@ -33,8 +33,8 @@
 
 UDMXPixelMappingMatrixComponent::UDMXPixelMappingMatrixComponent()
 {	
-	SizeX = 150.f;
-	SizeY = 150.f;
+	SizeX = 32.f;
+	SizeY = 32.f;
 
 	AttributeR.SetFromName("Red");
 	AttributeG.SetFromName("Green");
