DMX - Fix render component shows offset in pixelmapping UI

diff --git a/Engine/Plugins/VirtualProduction/DMX/DMXPixelMapping/Source/DMXPixelMappingEditor/Private/Views/SDMXPixelMappingDesignerView.cpp b/Engine/Plugins/VirtualProduction/DMX/DMXPixelMapping/Source/DMXPixelMappingEditor/Private/Views/SDMXPixelMappingDesignerView.cpp
--- a/Engine/Plugins/VirtualProduction/DMX/DMXPixelMapping/Source/DMXPixelMappingEditor/Private/Views/SDMXPixelMappingDesignerView.cpp
+++ b/Engine/Plugins/VirtualProduction/DMX/DMXPixelMapping/Source/DMXPixelMappingEditor/Private/Views/SDMXPixelMappingDesignerView.cpp
@@ -150,6 +150,8 @@ void SDMXPixelMappingDesignerView::Construct(const FArguments& InArgs, const TSh
 								SAssignNew(DesignCanvasBorder, SBorder)
 								.HAlign(HAlign_Fill)
 								.VAlign(VAlign_Fill)
+								.BorderImage(FEditorStyle::GetBrush("NoBorder"))
+								.Padding(0.0f)
 								[
 									SAssignNew(DesignCanvas, SConstraintCanvas)
 									.Visibility(EVisibility::HitTestInvisible)
