DMX - Remove a misleading log messages that would show when fixture types or patches are changed

diff --git a/Engine/Plugins/VirtualProduction/DMX/DMXEngine/Source/DMXRuntime/Public/Library/DMXEntityFixturePatch.h b/Engine/Plugins/VirtualProduction/DMX/DMXEngine/Source/DMXRuntime/Public/Library/DMXEntityFixturePatch.h
--- a/Engine/Plugins/VirtualProduction/DMX/DMXEngine/Source/DMXRuntime/Public/Library/DMXEntityFixturePatch.h
+++ b/Engine/Plugins/VirtualProduction/DMX/DMXEngine/Source/DMXRuntime/Public/Library/DMXEntityFixturePatch.h
@@ -46,7 +46,7 @@ struct DMXRUNTIME_API FDMXEntityFixturePatchConstructionParams
 	int32 UniverseID = 1;
 
 	/** Starting channel for when auto-assign address is false */
-	UPROPERTY(EditAnywhere, BlueprintReadWrite, Category = "Fixture Patch", meta = (EditCondition = "!bAutoAssignAddress", DisplayName = "Manual Starting Address", UIMin = "1", UIMax = "512", ClampMin = "1", ClampMax = "512"))
+	UPROPERTY(EditAnywhere, BlueprintReadWrite, Category = "Fixture Patch", meta = (DisplayName = "Starting Address", UIMin = "1", UIMax = "512", ClampMin = "1", ClampMax = "512"))
 	int32 StartingAddress = 1;
 };
 
diff --git a/Engine/Plugins/VirtualProduction/DMX/DMXPixelMapping/Source/DMXPixelMappingRuntime/Private/Components/DMXPixelMappingMatrixComponent.cpp b/Engine/Plugins/VirtualProduction/DMX/DMXPixelMapping/Source/DMXPixelMappingRuntime/Private/Components/DMXPixelMappingMatrixComponent.cpp
--- a/Engine/Plugins/VirtualProduction/DMX/DMXPixelMapping/Source/DMXPixelMappingRuntime/Private/Components/DMXPixelMappingMatrixComponent.cpp
+++ b/Engine/Plugins/VirtualProduction/DMX/DMXPixelMapping/Source/DMXPixelMappingRuntime/Private/Components/DMXPixelMappingMatrixComponent.cpp
@@ -46,10 +46,6 @@ UDMXPixelMappingMatrixComponent::UDMXPixelMappingMatrixComponent()
 	AttributeBExpose = true;
 
 	bMonochromeExpose = true;
-
-	// Listen to Fixture Type changes
-	UDMXEntityFixtureType::GetOnFixtureTypeChanged().AddUObject(this, &UDMXPixelMappingMatrixComponent::OnFixtureTypeChanged);
-	UDMXEntityFixturePatch::GetOnFixturePatchChanged().AddUObject(this, &UDMXPixelMappingMatrixComponent::OnFixturePatchChanged);
 }
 
 void UDMXPixelMappingMatrixComponent::Serialize(FArchive& Ar)
@@ -91,6 +87,16 @@ void UDMXPixelMappingMatrixComponent::PostLoad()
 	}
 }
 
+void UDMXPixelMappingMatrixComponent::PostInitProperties()
+{
+	if (!HasAnyFlags(RF_ClassDefaultObject | RF_ArchetypeObject))
+	{
+		// Listen to Fixture Type and Fixture Patch changes
+		UDMXEntityFixtureType::GetOnFixtureTypeChanged().AddUObject(this, &UDMXPixelMappingMatrixComponent::OnFixtureTypeChanged);
+		UDMXEntityFixturePatch::GetOnFixturePatchChanged().AddUObject(this, &UDMXPixelMappingMatrixComponent::OnFixturePatchChanged);
+	}
+}
+
 void UDMXPixelMappingMatrixComponent::LogInvalidProperties()
 {
 	UDMXEntityFixturePatch* FixturePatch = FixturePatchRef.GetFixturePatch();
diff --git a/Engine/Plugins/VirtualProduction/DMX/DMXPixelMapping/Source/DMXPixelMappingRuntime/Public/Components/DMXPixelMappingMatrixComponent.h b/Engine/Plugins/VirtualProduction/DMX/DMXPixelMapping/Source/DMXPixelMappingRuntime/Public/Components/DMXPixelMappingMatrixComponent.h
--- a/Engine/Plugins/VirtualProduction/DMX/DMXPixelMapping/Source/DMXPixelMappingRuntime/Public/Components/DMXPixelMappingMatrixComponent.h
+++ b/Engine/Plugins/VirtualProduction/DMX/DMXPixelMapping/Source/DMXPixelMappingRuntime/Public/Components/DMXPixelMappingMatrixComponent.h
@@ -40,6 +40,7 @@ public:
 	// ~Begin UObject interface
 	virtual void Serialize(FArchive& Ar) override;
 	virtual void PostLoad() override;
+	virtual void PostInitProperties() override;
 #if WITH_EDITOR
 	virtual void PostEditChangeProperty(FPropertyChangedEvent& PropertyChangedEvent) override;
 	virtual void PostEditChangeChainProperty(FPropertyChangedChainEvent& PropertyChangedChainEvent) override;
