DMX - Fix PixelMapping warnings are difficult to read

diff --git a/Engine/Plugins/VirtualProduction/DMX/DMXPixelMapping/Source/DMXPixelMappingEditor/DMXPixelMappingEditor.Build.cs b/Engine/Plugins/VirtualProduction/DMX/DMXPixelMapping/Source/DMXPixelMappingEditor/DMXPixelMappingEditor.Build.cs
--- a/Engine/Plugins/VirtualProduction/DMX/DMXPixelMapping/Source/DMXPixelMappingEditor/DMXPixelMappingEditor.Build.cs
+++ b/Engine/Plugins/VirtualProduction/DMX/DMXPixelMapping/Source/DMXPixelMappingEditor/DMXPixelMappingEditor.Build.cs
@@ -8,27 +8,27 @@ public class DMXPixelMappingEditor : ModuleRules
 	public DMXPixelMappingEditor(ReadOnlyTargetRules Target) : base(Target)
 	{
 		PrivateDependencyModuleNames.AddRange( new string[] {
+			"ApplicationCore",
 			"Core",
 			"CoreUObject",
-			"Engine",
-			"UnrealEd",
-			"Slate",
-			"SlateCore",
-			"EditorStyle",
-			"InputCore",
-			"PropertyEditor",
-			"RHI",
-			"RenderCore",
-			"EditorStyle",
-			"Projects",
-			"DMXRuntime",
 			"DMXEditor",
 			"DMXPixelMappingCore",
 			"DMXPixelMappingEditorWidgets",
 			"DMXPixelMappingRuntime",
 			"DMXPixelMappingRenderer",
 			"DMXPixelMappingBlueprintGraph",
-			"ApplicationCore"
+			"DMXRuntime",
+			"EditorStyle",
+			"Engine",
+			"InputCore",
+			"Projects",
+			"PropertyEditor",
+			"RenderCore",
+			"RHI",
+			"Slate",
+			"SlateCore",
+			"ToolWidgets",
+			"UnrealEd",
 		});
 	}
 }
diff --git a/Engine/Plugins/VirtualProduction/DMX/DMXPixelMapping/Source/DMXPixelMappingEditor/Private/Customizations/DMXPixelMappingDetailCustomization_Renderer.cpp b/Engine/Plugins/VirtualProduction/DMX/DMXPixelMapping/Source/DMXPixelMappingEditor/Private/Customizations/DMXPixelMappingDetailCustomization_Renderer.cpp
--- a/Engine/Plugins/VirtualProduction/DMX/DMXPixelMapping/Source/DMXPixelMappingEditor/Private/Customizations/DMXPixelMappingDetailCustomization_Renderer.cpp
+++ b/Engine/Plugins/VirtualProduction/DMX/DMXPixelMapping/Source/DMXPixelMappingEditor/Private/Customizations/DMXPixelMappingDetailCustomization_Renderer.cpp
@@ -7,6 +7,7 @@
 #include "DetailLayoutBuilder.h"
 #include "DetailCategoryBuilder.h"
 #include "DetailWidgetRow.h"
+#include "SWarningOrErrorBox.h"
 #include "Engine/Texture.h"
 #include "Layout/Visibility.h"
 #include "Materials/Material.h"
@@ -16,61 +17,6 @@
 
 #define LOCTEXT_NAMESPACE "DMXPixelMappingDetailCustomization_Renderer"
 
-/**
- * Renderer editor warning message widget
- */
-class SRendererCustomizationWarningMessage
-	: public SCompoundWidget
-{
-public:
-	SLATE_BEGIN_ARGS(SRendererCustomizationWarningMessage)
-		: _WarningText(FText::GetEmpty())
-		{}
-
-		SLATE_ATTRIBUTE(FText, WarningText)
-
-	SLATE_END_ARGS()
-
-	/**
-	 * Construct this widget
-	 *
-	 * @param	InArgs	The declaration data for this widget
-	 */
-	void Construct( const FArguments& InArgs )
-	{
-		const FSlateBrush* WarningIcon = FEditorStyle::GetBrush("SettingsEditor.WarningIcon");
-
-		ChildSlot
-		[
-			SNew(SBorder)
-			.BorderImage(FEditorStyle::GetBrush("SettingsEditor.CheckoutWarningBorder"))
-			.BorderBackgroundColor(FColor (166,137,0))
-			[
-				SNew(SHorizontalBox)
-				.Visibility(EVisibility::Visible)
-
-				+ SHorizontalBox::Slot()
-				.VAlign(VAlign_Center)
-				.AutoWidth()
-				.Padding(4.0f, 0.0f, 0.0f, 0.0f)
-				[
-					SNew(SImage)
-					.Image(WarningIcon)
-				]
-
-				+SHorizontalBox::Slot()
-				.VAlign(VAlign_Center)
-				.AutoWidth()
-				.Padding(4.0f, 0.0f, 0.0f, 0.0f)
-				[
-					SNew(STextBlock)
-					.Text(InArgs._WarningText)
-					.Font(IDetailLayoutBuilder::GetDetailFont())
-				]
-			]
-		];
-	}
-};
 
 void FDMXPixelMappingDetailCustomization_Renderer::CustomizeDetails(IDetailLayoutBuilder& DetailLayout)
 {
@@ -216,8 +162,9 @@ void FDMXPixelMappingDetailCustomization_Renderer::AddInputTextureWarning(IDetai
 			.Visibility(TAttribute<EVisibility>(this, &FDMXPixelMappingDetailCustomization_Renderer::GetInputTextureWarning))
 			.WholeRowContent()
 			[
-				SNew(SRendererCustomizationWarningMessage)
-					.WarningText(this, &FDMXPixelMappingDetailCustomization_Renderer::GetInputTextureWarningText)
+				SNew(SWarningOrErrorBox)
+				.MessageStyle(EMessageStyle::Warning)
+				.Message(this, &FDMXPixelMappingDetailCustomization_Renderer::GetInputTextureWarningText)
 			];
 }
 
@@ -227,8 +174,9 @@ void FDMXPixelMappingDetailCustomization_Renderer::AddMaterialWarning(IDetailCat
 			.Visibility(TAttribute<EVisibility>(this, &FDMXPixelMappingDetailCustomization_Renderer::GetMaterialWarningVisibility))
 			.WholeRowContent()
 			[
-				SNew(SRendererCustomizationWarningMessage)
-					.WarningText(LOCTEXT("WarningNonUIMaterial", "This is not UI Material.\nChange Material Domain to User Interface.\nOr select another Material."))
+				SNew(SWarningOrErrorBox)
+				.MessageStyle(EMessageStyle::Warning)
+				.Message(LOCTEXT("WarningNonUIMaterial", "Selected Material is not UI Material.\nChange Material Domain to User Interface.\nOr select another Material."))
 			];
 }
 
