Add missing transitions to DMXPixelMappingRenderer.

diff --git a/Engine/Plugins/VirtualProduction/DMX/DMXPixelMapping/Source/DMXPixelMappingRenderer/Private/DMXPixelMappingRenderer.cpp b/Engine/Plugins/VirtualProduction/DMX/DMXPixelMapping/Source/DMXPixelMappingRenderer/Private/DMXPixelMappingRenderer.cpp
--- a/Engine/Plugins/VirtualProduction/DMX/DMXPixelMapping/Source/DMXPixelMappingRenderer/Private/DMXPixelMappingRenderer.cpp
+++ b/Engine/Plugins/VirtualProduction/DMX/DMXPixelMapping/Source/DMXPixelMappingRenderer/Private/DMXPixelMappingRenderer.cpp
@@ -154,6 +154,8 @@ void FDMXPixelMappingRenderer::DownsampleRender(
 
 			const FTextureRHIRef InputTextureRHI = InputTexture->TextureRHI;
 
+			RHICmdList.Transition(FRHITransitionInfo(RenderTargetRef, ERHIAccess::SRVMask, ERHIAccess::RTV));
+
 			FRHIRenderPassInfo RpInfo(RenderTargetRef, ERenderTargetActions::Load_Store);
 			RHICmdList.BeginRenderPass(RpInfo, DMXPixelMappingRenderer::RenderPassName);
 			{
@@ -213,7 +215,10 @@ void FDMXPixelMappingRenderer::DownsampleRender(
 			// Copy texture from GPU to CPU
 			{
 				// Copies the contents of the given surface to its resolve target texture.
-				RHICmdList.CopyToResolveTarget(ResolveRenderTarget, RenderTargetRef, FResolveParams());
+				FResolveParams ResolveParams;
+				ResolveParams.SourceAccessFinal = ERHIAccess::SRVMask;
+				ResolveParams.DestAccessFinal = ERHIAccess::SRVMask;
+				RHICmdList.CopyToResolveTarget(ResolveRenderTarget, RenderTargetRef, ResolveParams);
 
 				// Read the contents of a texture to an output CPU buffer
 				TArray<FLinearColor> ColorArray;
@@ -395,6 +400,8 @@ void FDMXPixelMappingRenderer::RenderTextureToRectangle(const FTextureResource*
 	ENQUEUE_RENDER_COMMAND(DMXPixelMapping_CopyToPreveiewTexture)([this, RenderContext]
 	(FRHICommandListImmediate& RHICmdList)
 	{
+		RHICmdList.Transition(FRHITransitionInfo(RenderContext.TextureResource->TextureRHI, ERHIAccess::RTV, ERHIAccess::SRVMask));
+
 		FRHIRenderPassInfo RPInfo(RenderContext.Texture2DRHI, ERenderTargetActions::Load_Store);
 		RHICmdList.BeginRenderPass(RPInfo, TEXT("DMXPixelMapping_CopyToPreveiewTexture"));
 		{
@@ -428,8 +435,6 @@ void FDMXPixelMappingRenderer::RenderTextureToRectangle(const FTextureResource*
 				PixelShader->SetParameters(RHICmdList, TStaticSamplerState<SF_Point>::GetRHI(), RenderContext.TextureResource->TextureRHI);
 			}
 
-			
-
 			RendererModule->DrawRectangle(
 				RHICmdList,
 				0, 0,// X, Y Position in screen pixels of the top left corner of the quad
@@ -442,5 +447,7 @@ void FDMXPixelMappingRenderer::RenderTextureToRectangle(const FTextureResource*
 				EDRF_Default); // Flags see EDrawRectangleFlags
 		}
 		RHICmdList.EndRenderPass();
+
+		RHICmdList.Transition(FRHITransitionInfo(RenderContext.TextureResource->TextureRHI, ERHIAccess::SRVMask, ERHIAccess::RTV));
 	});
 }
