DMX - Fix Pixelmapping UI now correctly displays components on UNIX based systems

diff --git a/Engine/Plugins/VirtualProduction/DMX/DMXPixelMapping/Source/DMXPixelMappingRuntime/Private/Components/DMXPixelMappingBaseComponent.cpp b/Engine/Plugins/VirtualProduction/DMX/DMXPixelMapping/Source/DMXPixelMappingRuntime/Private/Components/DMXPixelMappingBaseComponent.cpp
--- a/Engine/Plugins/VirtualProduction/DMX/DMXPixelMapping/Source/DMXPixelMappingRuntime/Private/Components/DMXPixelMappingBaseComponent.cpp
+++ b/Engine/Plugins/VirtualProduction/DMX/DMXPixelMapping/Source/DMXPixelMappingRuntime/Private/Components/DMXPixelMappingBaseComponent.cpp
@@ -11,9 +11,28 @@
 #include "UObject/Package.h"
 
 
+UDMXPixelMappingBaseComponent::FDMXPixelMappingOnComponentAdded UDMXPixelMappingBaseComponent::OnComponentAdded;
+UDMXPixelMappingBaseComponent::FDMXPixelMappingOnComponentRemoved UDMXPixelMappingBaseComponent::OnComponentRemoved;
+UDMXPixelMappingBaseComponent::FDMXPixelMappingOnComponentRenamed UDMXPixelMappingBaseComponent::OnComponentRenamed;
+
 UDMXPixelMappingBaseComponent::UDMXPixelMappingBaseComponent()
 {}
 
+UDMXPixelMappingBaseComponent::FDMXPixelMappingOnComponentAdded& UDMXPixelMappingBaseComponent::UDMXPixelMappingBaseComponent::GetOnComponentAdded()
+{
+	return OnComponentAdded;
+}
+
+UDMXPixelMappingBaseComponent::FDMXPixelMappingOnComponentRemoved& UDMXPixelMappingBaseComponent::GetOnComponentRemoved()
+{
+	return OnComponentRemoved;
+}
+
+UDMXPixelMappingBaseComponent::FDMXPixelMappingOnComponentRenamed& UDMXPixelMappingBaseComponent::GetOnComponentRenamed()
+{
+	return OnComponentRenamed;
+}
+
 void UDMXPixelMappingBaseComponent::Serialize(FArchive& Ar)
 {
 	Super::Serialize(Ar);
diff --git a/Engine/Plugins/VirtualProduction/DMX/DMXPixelMapping/Source/DMXPixelMappingRuntime/Public/Components/DMXPixelMappingBaseComponent.h b/Engine/Plugins/VirtualProduction/DMX/DMXPixelMapping/Source/DMXPixelMappingRuntime/Public/Components/DMXPixelMappingBaseComponent.h
--- a/Engine/Plugins/VirtualProduction/DMX/DMXPixelMapping/Source/DMXPixelMappingRuntime/Public/Components/DMXPixelMappingBaseComponent.h
+++ b/Engine/Plugins/VirtualProduction/DMX/DMXPixelMapping/Source/DMXPixelMappingRuntime/Public/Components/DMXPixelMappingBaseComponent.h
@@ -34,26 +34,14 @@ public:
 	UDMXPixelMappingBaseComponent();
 
 	/** Gets an Event broadcast when a component was added */
-	static FDMXPixelMappingOnComponentAdded& GetOnComponentAdded()
-	{
-		static FDMXPixelMappingOnComponentAdded OnComponentAdded;
-		return OnComponentAdded;
-	}
+	static FDMXPixelMappingOnComponentAdded& GetOnComponentAdded();
 
 	/** Gets an Event broadcast when a component was added */
-	static FDMXPixelMappingOnComponentRemoved& GetOnComponentRemoved()
-	{
-		static FDMXPixelMappingOnComponentRemoved OnComponentRemoved;
-		return OnComponentRemoved;
-	}
+	static FDMXPixelMappingOnComponentRemoved& GetOnComponentRemoved();
 
 	/** Gets an Event broadcast when a component was renamed */
-	static FDMXPixelMappingOnComponentRenamed& GetOnComponentRenamed()
-	{
-		static FDMXPixelMappingOnComponentRenamed OnComponentRenamed;
-		return OnComponentRenamed;
-	}
-
+	static FDMXPixelMappingOnComponentRenamed& GetOnComponentRenamed();
+	
 protected:
 	//~ Begin UObject interface
 	virtual void Serialize(FArchive& Ar) override;
@@ -275,4 +263,13 @@ private:
 	/** Parent component */
 	UPROPERTY(NonTransactional)
 	TWeakObjectPtr<UDMXPixelMappingBaseComponent> WeakParent;
+
+	/** Delegate Broadcast when a component was added */
+	static FDMXPixelMappingOnComponentAdded OnComponentAdded;
+
+	/** Delegate Broadcast when a component was removed */
+	static FDMXPixelMappingOnComponentRemoved OnComponentRemoved;
+
+	/** Delegate Broadcast when a component was renamed */
+	static FDMXPixelMappingOnComponentRenamed OnComponentRenamed;
 };
