DMX - Add layout options for DMX Pixelmapping, component widgets now use a model view pattern

diff --git a/Engine/Plugins/VirtualProduction/DMX/DMXEngine/Source/DMXRuntime/Public/DMXOptionalTypes.h b/Engine/Plugins/VirtualProduction/DMX/DMXEngine/Source/DMXRuntime/Public/DMXOptionalTypes.h
--- a/Engine/Plugins/VirtualProduction/DMX/DMXEngine/Source/DMXRuntime/Public/DMXOptionalTypes.h
+++ b/Engine/Plugins/VirtualProduction/DMX/DMXEngine/Source/DMXRuntime/Public/DMXOptionalTypes.h
@@ -306,6 +306,84 @@ struct TStructOpsTypeTraits<FDMXOptionalGuid>
 	};
 };
 
+
+USTRUCT()
+struct DMXRUNTIME_API FDMXOptionalVector2D
+{
+	GENERATED_BODY()
+
+	/** The Optional Value */
+	TOptional<FVector2D> Value;
+
+	/** Resets the optional */
+	void Reset() { Value.Reset(); }
+	
+	/** Returns true if the optional is set */
+	bool IsSet() const { return Value.IsSet(); }
+	
+	/** Returns the Value of the optional, only call when the value is set. */
+	const FVector2D& GetValue() const { return Value.GetValue(); }
+
+	const FDMXOptionalVector2D& operator=(const FDMXOptionalVector2D& Rhs)
+	{
+		Value = Rhs.Value;
+		return *this;
+	}
+	const FDMXOptionalVector2D& operator=(const FVector2D& Vector2D)
+	{
+		Value = Vector2D;
+		return *this;
+	}
+
+	friend bool operator==(const FDMXOptionalVector2D& A, const FDMXOptionalVector2D& B)
+	{
+		const bool bBothSet = A.Value.IsSet() && B.Value.IsSet();
+		if (bBothSet)
+		{
+			return A.Value.GetValue() == B.Value.GetValue();
+		}
+		else
+		{
+			return A.Value.IsSet() == B.Value.IsSet();
+		}
+	}
+
+	friend bool operator!=(const FDMXOptionalVector2D& A, const FDMXOptionalVector2D& B)
+	{
+		return !(A == B);
+	}
+
+	bool Serialize(FArchive& Ar)
+	{
+		Ar << Value;
+		return true;
+	}
+
+	FArchive& operator<<(FArchive& Ar)
+	{
+		Ar << Value;
+		return Ar;
+	}
+
+	friend FArchive& operator<<(FArchive& Ar, FDMXOptionalVector2D& Struct)
+	{
+		Ar << Struct.Value;
+		return Ar;
+	}
+};
+	
+template<>
+struct TStructOpsTypeTraits<FDMXOptionalVector2D>
+	: public TStructOpsTypeTraitsBase2<FDMXOptionalVector2D>
+{
+	enum
+	{
+		WithSerializer = true,
+		WithCopy = true,
+		WithIdenticalViaEquality = true
+	};
+};
+
 USTRUCT()
 struct DMXRUNTIME_API FDMXOptionalTransform
 {
diff --git a/Engine/Plugins/VirtualProduction/DMX/DMXEngine/Source/DMXRuntime/Public/MVR/Types/DMXMVRChildListNode.h b/Engine/Plugins/VirtualProduction/DMX/DMXEngine/Source/DMXRuntime/Public/MVR/Types/DMXMVRChildListNode.h
--- a/Engine/Plugins/VirtualProduction/DMX/DMXEngine/Source/DMXRuntime/Public/MVR/Types/DMXMVRChildListNode.h
+++ b/Engine/Plugins/VirtualProduction/DMX/DMXEngine/Source/DMXRuntime/Public/MVR/Types/DMXMVRChildListNode.h
@@ -14,7 +14,7 @@ class FXmlNode;
 
 /** This node defines a generic graphical object. */
 UCLASS()
-class UDMXMVRChildListNode
+class DMXRUNTIME_API UDMXMVRChildListNode
 	: public UObject
 {
 	GENERATED_BODY()
diff --git a/Engine/Plugins/VirtualProduction/DMX/DMXEngine/Source/DMXRuntime/Public/MVR/Types/DMXMVRFixtureNode.h b/Engine/Plugins/VirtualProduction/DMX/DMXEngine/Source/DMXRuntime/Public/MVR/Types/DMXMVRFixtureNode.h
--- a/Engine/Plugins/VirtualProduction/DMX/DMXEngine/Source/DMXRuntime/Public/MVR/Types/DMXMVRFixtureNode.h
+++ b/Engine/Plugins/VirtualProduction/DMX/DMXEngine/Source/DMXRuntime/Public/MVR/Types/DMXMVRFixtureNode.h
@@ -143,7 +143,7 @@ struct TStructOpsTypeTraits<FDMXOptionalMVRFixtureGobo>
 
 /** This node defines a light fixture object. */
 UCLASS()
-class UDMXMVRFixtureNode
+class DMXRUNTIME_API UDMXMVRFixtureNode
 	: public UDMXMVRParametricObjectNodeBase
 {
 	GENERATED_BODY()
diff --git a/Engine/Plugins/VirtualProduction/DMX/DMXEngine/Source/DMXRuntime/Public/MVR/Types/DMXMVRGroupObjectNode.h b/Engine/Plugins/VirtualProduction/DMX/DMXEngine/Source/DMXRuntime/Public/MVR/Types/DMXMVRGroupObjectNode.h
--- a/Engine/Plugins/VirtualProduction/DMX/DMXEngine/Source/DMXRuntime/Public/MVR/Types/DMXMVRGroupObjectNode.h
+++ b/Engine/Plugins/VirtualProduction/DMX/DMXEngine/Source/DMXRuntime/Public/MVR/Types/DMXMVRGroupObjectNode.h
@@ -19,7 +19,7 @@ class FXmlNode;
 
 /** This node defines logical group of objects. The child objects are located inside a local coordinate system. */
 UCLASS()
-class UDMXMVRGroupObjectNode
+class DMXRUNTIME_API UDMXMVRGroupObjectNode
 	: public UDMXMVRParametricObjectNodeBase
 {
 	GENERATED_BODY()
diff --git a/Engine/Plugins/VirtualProduction/DMX/DMXEngine/Source/DMXRuntime/Public/MVR/Types/DMXMVRLayerNode.h b/Engine/Plugins/VirtualProduction/DMX/DMXEngine/Source/DMXRuntime/Public/MVR/Types/DMXMVRLayerNode.h
--- a/Engine/Plugins/VirtualProduction/DMX/DMXEngine/Source/DMXRuntime/Public/MVR/Types/DMXMVRLayerNode.h
+++ b/Engine/Plugins/VirtualProduction/DMX/DMXEngine/Source/DMXRuntime/Public/MVR/Types/DMXMVRLayerNode.h
@@ -17,7 +17,7 @@ class FXmlNode;
 
 /** A Layer node in the Layers Node of a Scene. */
 UCLASS()
-class UDMXMVRLayerNode
+class DMXRUNTIME_API UDMXMVRLayerNode
 	: public UObject
 {
 	GENERATED_BODY()
diff --git a/Engine/Plugins/VirtualProduction/DMX/DMXEngine/Source/DMXRuntime/Public/MVR/Types/DMXMVRParametricObjectNodeBase.h b/Engine/Plugins/VirtualProduction/DMX/DMXEngine/Source/DMXRuntime/Public/MVR/Types/DMXMVRParametricObjectNodeBase.h
--- a/Engine/Plugins/VirtualProduction/DMX/DMXEngine/Source/DMXRuntime/Public/MVR/Types/DMXMVRParametricObjectNodeBase.h
+++ b/Engine/Plugins/VirtualProduction/DMX/DMXEngine/Source/DMXRuntime/Public/MVR/Types/DMXMVRParametricObjectNodeBase.h
@@ -18,7 +18,7 @@ class FXmlNode;
 
 /** UE specific: The base class for all Children in a Child List Node, whereas the standard refers to them as 'Parametric Objects'. */
 UCLASS(Abstract)
-class UDMXMVRParametricObjectNodeBase
+class DMXRUNTIME_API UDMXMVRParametricObjectNodeBase
 	: public UObject
 {
 	GENERATED_BODY()
diff --git a/Engine/Plugins/VirtualProduction/DMX/DMXEngine/Source/DMXRuntime/Public/MVR/Types/DMXMVRRootNode.h b/Engine/Plugins/VirtualProduction/DMX/DMXEngine/Source/DMXRuntime/Public/MVR/Types/DMXMVRRootNode.h
--- a/Engine/Plugins/VirtualProduction/DMX/DMXEngine/Source/DMXRuntime/Public/MVR/Types/DMXMVRRootNode.h
+++ b/Engine/Plugins/VirtualProduction/DMX/DMXEngine/Source/DMXRuntime/Public/MVR/Types/DMXMVRRootNode.h
@@ -22,7 +22,7 @@ class FXmlFile;
  * UE specific: Note, while properties follow the standard closely, only children relevant to the Engine are implemented.
  */
 UCLASS()
-class UDMXMVRRootNode
+class DMXRUNTIME_API UDMXMVRRootNode
 	: public UObject
 {
 	GENERATED_BODY()
diff --git a/Engine/Plugins/VirtualProduction/DMX/DMXEngine/Source/DMXRuntime/Public/MVR/Types/DMXMVRSceneNode.h b/Engine/Plugins/VirtualProduction/DMX/DMXEngine/Source/DMXRuntime/Public/MVR/Types/DMXMVRSceneNode.h
--- a/Engine/Plugins/VirtualProduction/DMX/DMXEngine/Source/DMXRuntime/Public/MVR/Types/DMXMVRSceneNode.h
+++ b/Engine/Plugins/VirtualProduction/DMX/DMXEngine/Source/DMXRuntime/Public/MVR/Types/DMXMVRSceneNode.h
@@ -17,7 +17,7 @@ class FXmlNode;
  * UE Specific: Note, while properties follow the standard closely, only properties relevant to the Engine are implemented.
  */
 UCLASS()
-class UDMXMVRSceneNode
+class DMXRUNTIME_API UDMXMVRSceneNode
 	: public UObject
 {
 	GENERATED_BODY()
diff --git a/Engine/Plugins/VirtualProduction/DMX/DMXEngine/Source/DMXRuntime/Public/MVR/Types/DMXMVRUnrealEngineDataNode.h b/Engine/Plugins/VirtualProduction/DMX/DMXEngine/Source/DMXRuntime/Public/MVR/Types/DMXMVRUnrealEngineDataNode.h
--- a/Engine/Plugins/VirtualProduction/DMX/DMXEngine/Source/DMXRuntime/Public/MVR/Types/DMXMVRUnrealEngineDataNode.h
+++ b/Engine/Plugins/VirtualProduction/DMX/DMXEngine/Source/DMXRuntime/Public/MVR/Types/DMXMVRUnrealEngineDataNode.h
@@ -14,7 +14,7 @@ class FXmlNode;
  * UE sepecific: Used to add UE specific meta data
  */
 UCLASS()
-class UDMXMVRUnrealEngineDataNode
+class DMXRUNTIME_API UDMXMVRUnrealEngineDataNode
 	: public UObject
 {
 	GENERATED_BODY()
diff --git a/Engine/Plugins/VirtualProduction/DMX/DMXEngine/Source/DMXRuntime/Public/MVR/Types/DMXMVRUserDataNode.h b/Engine/Plugins/VirtualProduction/DMX/DMXEngine/Source/DMXRuntime/Public/MVR/Types/DMXMVRUserDataNode.h
--- a/Engine/Plugins/VirtualProduction/DMX/DMXEngine/Source/DMXRuntime/Public/MVR/Types/DMXMVRUserDataNode.h
+++ b/Engine/Plugins/VirtualProduction/DMX/DMXEngine/Source/DMXRuntime/Public/MVR/Types/DMXMVRUserDataNode.h
@@ -14,7 +14,7 @@ class FXmlNode;
 
 /** This node contains a collection of user data nodes defined and used by provider applications if required. */
 UCLASS()
-class UDMXMVRUserDataNode
+class DMXRUNTIME_API UDMXMVRUserDataNode
 	: public UObject
 {
 	GENERATED_BODY()
diff --git a/Engine/Plugins/VirtualProduction/DMX/DMXPixelMapping/Source/DMXPixelMappingBlueprintGraph/Private/K2Node_PixelMappingRendererComponent.cpp b/Engine/Plugins/VirtualProduction/DMX/DMXPixelMapping/Source/DMXPixelMappingBlueprintGraph/Private/K2Node_PixelMappingRendererComponent.cpp
--- a/Engine/Plugins/VirtualProduction/DMX/DMXPixelMapping/Source/DMXPixelMappingBlueprintGraph/Private/K2Node_PixelMappingRendererComponent.cpp
+++ b/Engine/Plugins/VirtualProduction/DMX/DMXPixelMapping/Source/DMXPixelMappingBlueprintGraph/Private/K2Node_PixelMappingRendererComponent.cpp
@@ -56,18 +56,22 @@ void UK2Node_PixelMappingRendererComponent::GetMenuActions(FBlueprintActionDatab
 
 UEdGraphPin* UK2Node_PixelMappingRendererComponent::GetInRendererComponentPin() const
 {
-	UEdGraphPin* Pin = FindPinChecked(InRendererComponentPinName);
-	check(Pin->Direction == EGPD_Input);
-
-	return Pin;
+	if (UEdGraphPin* Pin = FindPin(InRendererComponentPinName))
+	{
+		check(Pin->Direction == EGPD_Input);
+		return Pin;
+	}
+	return nullptr;
 }
 
 UEdGraphPin* UK2Node_PixelMappingRendererComponent::GetOutRendererComponentPin() const
 {
-	UEdGraphPin* Pin = FindPinChecked(OutRendererComponentPinName);
-	check(Pin->Direction == EGPD_Output);
-
-	return Pin;
+	if (UEdGraphPin* Pin = FindPin(OutRendererComponentPinName))
+	{
+		check(Pin->Direction == EGPD_Output);
+		return Pin;
+	}
+	return nullptr;
 }
 
 void UK2Node_PixelMappingRendererComponent::EarlyValidation(FCompilerResultsLog& MessageLog) const
diff --git a/Engine/Plugins/VirtualProduction/DMX/DMXPixelMapping/Source/DMXPixelMappingEditor/Private/Customizations/DMXPixelMappingDetailCustomization_FixtureGroup.cpp b/Engine/Plugins/VirtualProduction/DMX/DMXPixelMapping/Source/DMXPixelMappingEditor/Private/Customizations/DMXPixelMappingDetailCustomization_FixtureGroup.cpp
--- a/Engine/Plugins/VirtualProduction/DMX/DMXPixelMapping/Source/DMXPixelMappingEditor/Private/Customizations/DMXPixelMappingDetailCustomization_FixtureGroup.cpp
+++ b/Engine/Plugins/VirtualProduction/DMX/DMXPixelMapping/Source/DMXPixelMappingEditor/Private/Customizations/DMXPixelMappingDetailCustomization_FixtureGroup.cpp
@@ -2,6 +2,8 @@
 
 #include "Customizations/DMXPixelMappingDetailCustomization_FixtureGroup.h"
 
+#include "DMXPixelMapping.h"
+#include "DMXPixelMappingLayoutSettings.h"
 #include "Components/DMXPixelMappingFixtureGroupComponent.h"
 #include "Components/DMXPixelMappingFixtureGroupItemComponent.h"
 #include "Components/DMXPixelMappingMatrixComponent.h"
@@ -21,9 +23,10 @@
 #include "ScopedTransaction.h"
 #include "Layout/Visibility.h"
 #include "Misc/CoreDelegates.h"
-#include "Widgets/Views/SListView.h"
-#include "Widgets/Layout/SScrollBorder.h"
 #include "Widgets/SBoxPanel.h"
+#include "Widgets/Input/SButton.h"
+#include "Widgets/Layout/SScrollBorder.h"
+#include "Widgets/Views/SListView.h"
 
 
 #define LOCTEXT_NAMESPACE "DMXPixelMappingDetailCustomization_FixtureGroup"
@@ -32,72 +35,206 @@ void FDMXPixelMappingDetailCustomization_FixtureGroup::CustomizeDetails(IDetailL
 {
 	PropertyUtilities = InDetailLayout.GetPropertyUtilities();
 
+	// Hide the Layout Script property (shown in its own panel, see SDMXPixelMappingLayoutView)
+	InDetailLayout.HideProperty(GET_MEMBER_NAME_CHECKED(UDMXPixelMappingFixtureGroupComponent, LayoutScript));
+
+	// Handle size changes
+	UpdateCachedScaleChildrenWithParent();
+	SizeXHandle = InDetailLayout.GetProperty(UDMXPixelMappingOutputComponent::GetSizeXPropertyName(), UDMXPixelMappingOutputComponent::StaticClass());
+	SizeXHandle->SetOnPropertyValuePreChange(FSimpleDelegate::CreateSP(this, &FDMXPixelMappingDetailCustomization_FixtureGroup::OnSizePropertyPreChange));
+	SizeXHandle->SetOnPropertyValueChangedWithData(TDelegate<void(const FPropertyChangedEvent&)>::CreateSP(this, &FDMXPixelMappingDetailCustomization_FixtureGroup::OnSizePropertyChanged));
+
+	SizeYHandle = InDetailLayout.GetProperty(UDMXPixelMappingOutputComponent::GetSizeYPropertyName(), UDMXPixelMappingOutputComponent::StaticClass());
+	SizeYHandle->SetOnPropertyValuePreChange(FSimpleDelegate::CreateSP(this, &FDMXPixelMappingDetailCustomization_FixtureGroup::OnSizePropertyPreChange));
+	SizeYHandle->SetOnPropertyValueChangedWithData(TDelegate<void(const FPropertyChangedEvent&)>::CreateSP(this, &FDMXPixelMappingDetailCustomization_FixtureGroup::OnSizePropertyChanged));
+
+	// Remember the group being edited
 	WeakFixtureGroupComponent = GetSelectedFixtureGroupComponent(InDetailLayout);
-	if (UDMXPixelMappingFixtureGroupComponent* FixtureGroupComponent = WeakFixtureGroupComponent.Get())
+
+	// Listen to the library being changed in the group component
+	DMXLibraryHandle = InDetailLayout.GetProperty(GET_MEMBER_NAME_CHECKED(UDMXPixelMappingFixtureGroupComponent, DMXLibrary));
+	DMXLibraryHandle->SetOnPropertyValueChanged(FSimpleDelegate::CreateSP(this, &FDMXPixelMappingDetailCustomization_FixtureGroup::OnLibraryChanged));
+	DMXLibraryHandle->SetOnChildPropertyValueChanged(FSimpleDelegate::CreateSP(this, &FDMXPixelMappingDetailCustomization_FixtureGroup::OnLibraryChanged));
+
+	// Listen to component changes
+	UDMXPixelMappingBaseComponent::GetOnComponentAdded().AddSP(this, &FDMXPixelMappingDetailCustomization_FixtureGroup::OnComponentAdded);
+	UDMXPixelMappingBaseComponent::GetOnComponentRemoved().AddSP(this, &FDMXPixelMappingDetailCustomization_FixtureGroup::OnComponentRemoved);
+
+	// Add the library property
+	IDetailCategoryBuilder& FixtureListCategoryBuilder = InDetailLayout.EditCategory("Fixture List", FText::GetEmpty(), ECategoryPriority::Important);
+	FixtureListCategoryBuilder.AddProperty(DMXLibraryHandle);
+
+	// Add an 'Add all Patches' button
+	CreateAddAllPatchesButton(InDetailLayout);
+
+	// Add Fixture Patches to the view
+	CreateFixturePatchDetailRows(InDetailLayout);
+}
+
+void FDMXPixelMappingDetailCustomization_FixtureGroup::CreateAddAllPatchesButton(IDetailLayoutBuilder& InDetailLayout)
+{
+	IDetailCategoryBuilder& FixtureListCategoryBuilder = InDetailLayout.EditCategory("Fixture List", FText::GetEmpty(), ECategoryPriority::Important);
+
+	FixtureListCategoryBuilder.AddCustomRow(FText::GetEmpty())
+		.WholeRowContent()
+		[
+			SNew(SButton)
+			.Text(LOCTEXT("AddAllPatchesButtonText", "Add all Patches"))
+			.OnClicked(this, &FDMXPixelMappingDetailCustomization_FixtureGroup::OnAddAllPatchesClicked)
+			.Visibility_Lambda([this]()
+				{
+					return NumFixturePatchRows > 0 ? EVisibility::Visible : EVisibility::Collapsed;
+				})
+		];
+}
+
+void FDMXPixelMappingDetailCustomization_FixtureGroup::CreateFixturePatchDetailRows(IDetailLayoutBuilder& InDetailLayout)
+{
+	UDMXPixelMappingFixtureGroupComponent* FixtureGroupComponent = WeakFixtureGroupComponent.Get();
+	if (!FixtureGroupComponent)
 	{
-		// Listen to component changes
-		UDMXPixelMappingBaseComponent::GetOnComponentAdded().AddSP(this, &FDMXPixelMappingDetailCustomization_FixtureGroup::OnComponentAdded);
-		UDMXPixelMappingBaseComponent::GetOnComponentRemoved().AddSP(this, &FDMXPixelMappingDetailCustomization_FixtureGroup::OnComponentRemoved);
+		return;
+	}
+
+	UDMXLibrary* DMXLibrary = GetSelectedDMXLibrary(FixtureGroupComponent);
+	if (!DMXLibrary)
+	{
+		return;
+	}
 
-		// Listen to the library being changed in the group component
-		DMXLibraryHandle = InDetailLayout.GetProperty(GET_MEMBER_NAME_CHECKED(UDMXPixelMappingFixtureGroupComponent, DMXLibrary), UDMXPixelMappingFixtureGroupComponent::StaticClass());
-		DMXLibraryHandle->SetOnPropertyValueChanged(FSimpleDelegate::CreateSP(this, &FDMXPixelMappingDetailCustomization_FixtureGroup::OnLibraryChanged));
-		DMXLibraryHandle->SetOnChildPropertyValueChanged(FSimpleDelegate::CreateSP(this, &FDMXPixelMappingDetailCustomization_FixtureGroup::OnLibraryChanged));
+	UpdateFixturePatchesInUse(DMXLibrary);
 
-		if (UDMXLibrary* DMXLibrary = GetSelectedDMXLibrary(FixtureGroupComponent))
+	// Listen to the entities array being changed in the library
+	EntitiesHandle = InDetailLayout.GetProperty(UDMXLibrary::GetEntitiesPropertyName());
+	EntitiesHandle->SetOnPropertyValueChanged(FSimpleDelegate::CreateSP(this, &FDMXPixelMappingDetailCustomization_FixtureGroup::ForceRefresh));
+
+	// Add fixture patches as custom rows
+	TArray<UDMXEntityFixturePatch*> AllFixturePatches = DMXLibrary->GetEntitiesTypeCast<UDMXEntityFixturePatch>();
+	for (UDMXEntityFixturePatch* FixturePatch : AllFixturePatches)
+	{
+		const bool bPatchIsAssigned = IsFixturePatchAssignedToPixelMapping(FixturePatch);
+
+		if (!bPatchIsAssigned)
 		{
-			UpdateFixturePatchesInUse(DMXLibrary);
+			TSharedRef<SDMXPixelMappingFixturePatchDetailRow> FixturePatchDetailRowWidget =
+				SNew(SDMXPixelMappingFixturePatchDetailRow)
+				.FixturePatch(FixturePatch)
+				.OnLMBDown(this, &FDMXPixelMappingDetailCustomization_FixtureGroup::OnFixturePatchLMBDown, FDMXEntityFixturePatchRef(FixturePatch))
+				.OnLMBUp(this, &FDMXPixelMappingDetailCustomization_FixtureGroup::OnFixturePatchLMBUp, FDMXEntityFixturePatchRef(FixturePatch))
+				.OnDragged(this, &FDMXPixelMappingDetailCustomization_FixtureGroup::OnFixturePatchesDragged);
+
+			FDMXPixelMappingDetailCustomization_FixtureGroup::FDetailRowWidgetWithPatch DetailRowWidgetWithPatch;
+			DetailRowWidgetWithPatch.DetailRowWidget = FixturePatchDetailRowWidget;
+			DetailRowWidgetWithPatch.WeakFixturePatch = FixturePatch;
+			DetailRowWidgetsWithPatch.Add(DetailRowWidgetWithPatch);
 
-			// Get editing categories
 			IDetailCategoryBuilder& FixtureListCategoryBuilder = InDetailLayout.EditCategory("Fixture List", FText::GetEmpty(), ECategoryPriority::Important);
+			FixtureListCategoryBuilder.AddCustomRow(FText::GetEmpty())
+				.WholeRowContent()
+				[
+					FixturePatchDetailRowWidget
+				];
 
-			// Listen to the entities array being changed in the library
-			EntitiesHandle = InDetailLayout.GetProperty(UDMXLibrary::GetEntitiesPropertyName());
-			EntitiesHandle->SetOnPropertyValueChanged(FSimpleDelegate::CreateSP(this, &FDMXPixelMappingDetailCustomization_FixtureGroup::ForceRefresh));
+			NumFixturePatchRows++;
+		}
+	}
+}
 
-			// Add the library property
-			FixtureListCategoryBuilder.AddProperty(DMXLibraryHandle);
+FReply FDMXPixelMappingDetailCustomization_FixtureGroup::OnAddAllPatchesClicked()
+{
+	const TSharedPtr<FDMXPixelMappingToolkit> Toolkit = ToolkitWeakPtr.Pin();
+	if (!Toolkit.IsValid())
+	{
+		return FReply::Handled();
+	}
 
-			// Add fixture patches as custom rows
-			TArray<UDMXEntityFixturePatch*> AllFixturePatches = DMXLibrary->GetEntitiesTypeCast<UDMXEntityFixturePatch>();
-			for (UDMXEntityFixturePatch* FixturePatch : AllFixturePatches)
-			{
-				const bool bPatchIsAssigned = FixtureGroupComponent->Children.ContainsByPredicate([FixturePatch](UDMXPixelMappingBaseComponent* BaseComponent)
-					{
-						if (UDMXPixelMappingFixtureGroupItemComponent* GroupItemComponent = Cast<UDMXPixelMappingFixtureGroupItemComponent>(BaseComponent))
-						{
-							return GroupItemComponent->FixturePatchRef.GetFixturePatch() == FixturePatch;
-						}
-						else if (UDMXPixelMappingMatrixComponent* MatrixComponent = Cast<UDMXPixelMappingMatrixComponent>(BaseComponent))
-						{
-							return MatrixComponent->FixturePatchRef.GetFixturePatch() == FixturePatch;
-						}
-						return false;
-					});
+	UDMXPixelMappingFixtureGroupComponent* GroupComponent = WeakFixtureGroupComponent.Get();
+	if (!GroupComponent)
+	{
+		return FReply::Handled();
+	}
 
-				if (!bPatchIsAssigned)
-				{
-					TSharedRef<SDMXPixelMappingFixturePatchDetailRow> FixturePatchDetailRowWidget =
-						SNew(SDMXPixelMappingFixturePatchDetailRow)
-						.FixturePatch(FixturePatch)
-						.OnLMBDown(this, &FDMXPixelMappingDetailCustomization_FixtureGroup::OnFixturePatchLMBDown, FDMXEntityFixturePatchRef(FixturePatch))
-						.OnLMBUp(this, &FDMXPixelMappingDetailCustomization_FixtureGroup::OnFixturePatchLMBUp, FDMXEntityFixturePatchRef(FixturePatch))
-						.OnDragged(this, &FDMXPixelMappingDetailCustomization_FixtureGroup::OnFixturePatchesDragged);
-
-					FDMXPixelMappingDetailCustomization_FixtureGroup::FDetailRowWidgetWithPatch DetailRowWidgetWithPatch;
-					DetailRowWidgetWithPatch.DetailRowWidget = FixturePatchDetailRowWidget;
-					DetailRowWidgetWithPatch.WeakFixturePatch = FixturePatch;
-					DetailRowWidgetsWithPatch.Add(DetailRowWidgetWithPatch);
-
-					FixtureListCategoryBuilder.AddCustomRow(FText::GetEmpty())
-						.WholeRowContent()
-						[
-							FixturePatchDetailRowWidget
-						];
-				}
-			}
+	UDMXLibrary* DMXLibrary = GetSelectedDMXLibrary(GroupComponent);
+	if (!DMXLibrary)
+	{
+		return FReply::Handled();
+	}
+
+	UDMXPixelMapping* PixelMapping = GroupComponent->GetPixelMapping();
+	if (!PixelMapping)
+	{
+		return FReply::Handled();
+	}
+
+	UDMXPixelMappingRootComponent* RootComponent = PixelMapping->GetRootComponent();
+	if (!RootComponent)
+	{
+		return FReply::Handled();
+	}
+
+	// Create components
+	FScopedTransaction AddAllPatchesTransaction(LOCTEXT("AddAllPatchesTransaction", "Add all Patches to Pixel Mapping"));
+	GroupComponent->Modify();
+
+	TArray<UDMXPixelMappingOutputComponent*> NewComponents;
+	for (const FDMXEntityFixturePatchRef& FixturePatchRef : FixturePatches)
+	{
+		UDMXEntityFixturePatch* FixturePatch = FixturePatchRef.GetFixturePatch();
+		if (!FixturePatch)
+		{
+			continue;
+		}
+
+		const bool bPatchIsAssigned = IsFixturePatchAssignedToPixelMapping(FixturePatch);
+		if (bPatchIsAssigned)
+		{
+			continue;
+		}
+
+		const FDMXFixtureMode* FixtureModePtr = FixturePatch->GetActiveMode();
+		if (FixtureModePtr && FixtureModePtr->bFixtureMatrixEnabled)
+		{
+			const TSharedPtr<FDMXPixelMappingComponentTemplate> ComponentTemplate = MakeShared<FDMXPixelMappingComponentTemplate>(UDMXPixelMappingMatrixComponent::StaticClass(), FixturePatchRef);
+			UDMXPixelMappingMatrixComponent* NewMatrixComponent = ComponentTemplate->CreateComponent<UDMXPixelMappingMatrixComponent>(RootComponent);
+			NewComponents.Add(NewMatrixComponent);
+
+			GroupComponent->AddChild(NewMatrixComponent);
+		}
+		else
+		{
+			const TSharedPtr<FDMXPixelMappingComponentTemplate> ComponentTemplate = MakeShared<FDMXPixelMappingComponentTemplate>(UDMXPixelMappingFixtureGroupItemComponent::StaticClass(), FixturePatchRef);
+			UDMXPixelMappingFixtureGroupItemComponent* NewGroupItemComponent = ComponentTemplate->CreateComponent<UDMXPixelMappingFixtureGroupItemComponent>(RootComponent);
+			NewComponents.Add(NewGroupItemComponent);
+
+			GroupComponent->AddChild(NewGroupItemComponent);
 		}
 	}
+
+	// Layout new components inside the group
+	const int32 Columns = FMath::Sqrt((float)NewComponents.Num());
+	const int32 Rows = FMath::RoundFromZero((float)NewComponents.Num() / Columns);
+	const float Size = FMath::Min(GroupComponent->GetSize().X / Columns, GroupComponent->GetSize().Y / Rows);
+	if (Size < 1.f)
+	{
+		GroupComponent->SetSize(FVector2D(Size * Columns, Size * Rows));
+	}
+	
+	const FVector2D ParentPosition = GroupComponent->GetPosition();
+	int32 Column = -1;
+	int32 Row = 0;
+	for (UDMXPixelMappingOutputComponent* NewComponent : NewComponents)
+	{
+		Column++;
+		if (Column > Columns)
+		{
+			Column = 0;
+			Row++;
+		}
+		const FVector2D Position = ParentPosition + FVector2D(Column * Size, Row * Size);
+		NewComponent->SetPosition(Position);
+	}
+
+	return FReply::Handled();
 }
 
 void FDMXPixelMappingDetailCustomization_FixtureGroup::OnLibraryChanged()
@@ -154,6 +291,106 @@ void FDMXPixelMappingDetailCustomization_FixtureGroup::OnComponentRemoved(UDMXPi
 	}
 }
 
+void FDMXPixelMappingDetailCustomization_FixtureGroup::OnSizePropertyPreChange()
+{
+	UpdateCachedScaleChildrenWithParent();
+	if (!bCachedScaleChildrenWithParent)
+	{
+		return;
+	}
+
+	const TArray<TWeakObjectPtr<UObject>> SelectedObjects = PropertyUtilities->GetSelectedObjects();
+	TArray<UDMXPixelMappingFixtureGroupComponent*> FixtureGroupComponents;
+	for (TWeakObjectPtr<UObject> Object : SelectedObjects)
+	{
+		if (UDMXPixelMappingFixtureGroupComponent* FixtureGroupComponent = Cast<UDMXPixelMappingFixtureGroupComponent>(Object))
+		{
+			FixtureGroupComponents.Add(FixtureGroupComponent);
+		}
+	}
+	if (FixtureGroupComponents.IsEmpty())
+	{
+		return;
+	}
+
+	PreEditChangeComponentToSizeMap.Reset();
+	for (UDMXPixelMappingFixtureGroupComponent* FixtureGroupComponent : FixtureGroupComponents)
+	{
+		PreEditChangeComponentToSizeMap.Add(FixtureGroupComponent, FixtureGroupComponent->GetSize());
+	}
+}
+
+void FDMXPixelMappingDetailCustomization_FixtureGroup::OnSizePropertyChanged(const FPropertyChangedEvent& PropertyChangedEvent)
+{
+	if(PropertyChangedEvent.ChangeType == EPropertyChangeType::Interactive)
+	{
+		GEditor->GetTimerManager()->SetTimerForNextTick(FSimpleDelegate::CreateSP(this, &FDMXPixelMappingDetailCustomization_FixtureGroup::HandleSizePropertyChanged));
+	}
+	else
+	{
+		HandleSizePropertyChanged();
+	}
+}
+
+void FDMXPixelMappingDetailCustomization_FixtureGroup::HandleSizePropertyChanged()
+{
+	// Scale children if desired
+	if (!bCachedScaleChildrenWithParent || 
+		PreEditChangeComponentToSizeMap.IsEmpty())
+	{
+		return;
+	}
+
+	const TArray<TWeakObjectPtr<UObject>> SelectedObjects = PropertyUtilities->GetSelectedObjects();
+	TArray<UDMXPixelMappingFixtureGroupComponent*> FixtureGroupComponents;
+	for (TWeakObjectPtr<UObject> Object : SelectedObjects)
+	{
+		if (UDMXPixelMappingFixtureGroupComponent* FixtureGroupComponent = Cast<UDMXPixelMappingFixtureGroupComponent>(Object))
+		{
+			FixtureGroupComponents.Add(FixtureGroupComponent);
+		}
+	}
+	if (FixtureGroupComponents.IsEmpty())
+	{
+		return;
+	}
+
+	for (const TTuple<TWeakObjectPtr<UDMXPixelMappingFixtureGroupComponent>, FVector2D>& PreEditChangeComponentToSizeXPair : PreEditChangeComponentToSizeMap)
+	{
+		UDMXPixelMappingFixtureGroupComponent* const* ComponentPtr = Algo::Find(FixtureGroupComponents, PreEditChangeComponentToSizeXPair.Key.Get());
+		if (!ComponentPtr)
+		{
+			continue;
+		}
+
+		const FVector2D GroupPosition = (*ComponentPtr)->GetPosition();
+		const FVector2D OldSize = PreEditChangeComponentToSizeXPair.Value;
+		const FVector2D NewSize = (*ComponentPtr)->GetSize();
+		if (NewSize == FVector2D::ZeroVector || OldSize == NewSize)
+		{
+			// No division by zero, no unchanged values
+			return;
+		}
+
+		const FVector2D RatioVector = NewSize / OldSize;
+		for (UDMXPixelMappingBaseComponent* BaseChild : (*ComponentPtr)->GetChildren())
+		{
+			if (UDMXPixelMappingOutputComponent* Child = Cast<UDMXPixelMappingOutputComponent>(BaseChild))
+			{
+				Child->Modify();
+
+				// Scale size (SetSize already clamps)
+				Child->SetSize(Child->GetSize() * RatioVector);
+
+				// Scale position
+				const FVector2D ChildPosition = Child->GetPosition();
+				const FVector2D NewPositionRelative = (ChildPosition - GroupPosition) * RatioVector;
+				Child->SetPosition(GroupPosition + NewPositionRelative);
+			}
+		}
+	}
+}
+
 void FDMXPixelMappingDetailCustomization_FixtureGroup::ForceRefresh()
 {
 	// Reset the handles so they won't fire any changes after refreshing
@@ -285,6 +522,29 @@ FReply FDMXPixelMappingDetailCustomization_FixtureGroup::OnFixturePatchesDragged
 	return FReply::Handled();
 }
 
+bool FDMXPixelMappingDetailCustomization_FixtureGroup::IsFixturePatchAssignedToPixelMapping(UDMXEntityFixturePatch* FixturePatch) const
+{
+	if (!WeakFixtureGroupComponent.IsValid())
+	{
+		return true;
+	}
+
+	const bool bPatchIsAssigned = WeakFixtureGroupComponent->Children.ContainsByPredicate([FixturePatch](UDMXPixelMappingBaseComponent* BaseComponent)
+		{
+			if (UDMXPixelMappingFixtureGroupItemComponent* GroupItemComponent = Cast<UDMXPixelMappingFixtureGroupItemComponent>(BaseComponent))
+			{
+				return GroupItemComponent->FixturePatchRef.GetFixturePatch() == FixturePatch;
+			}
+			else if (UDMXPixelMappingMatrixComponent* MatrixComponent = Cast<UDMXPixelMappingMatrixComponent>(BaseComponent))
+			{
+				return MatrixComponent->FixturePatchRef.GetFixturePatch() == FixturePatch;
+			}
+			return false;
+		});
+
+	return bPatchIsAssigned;
+}
+
 void FDMXPixelMappingDetailCustomization_FixtureGroup::UpdateFixturePatchHighlights()
 {
 	for (const FDMXPixelMappingDetailCustomization_FixtureGroup::FDetailRowWidgetWithPatch& DetailRowWidgetWithPatch : DetailRowWidgetsWithPatch)
@@ -317,6 +577,15 @@ void FDMXPixelMappingDetailCustomization_FixtureGroup::UpdateFixturePatchesInUse
 		});
 }
 
+void FDMXPixelMappingDetailCustomization_FixtureGroup::UpdateCachedScaleChildrenWithParent()
+{
+	const UDMXPixelMappingLayoutSettings* LayoutSettings = GetDefault<UDMXPixelMappingLayoutSettings>();
+	if (LayoutSettings)
+	{
+		bCachedScaleChildrenWithParent = LayoutSettings->bScaleChildrenWithParent;
+	}
+}
+
 UDMXLibrary* FDMXPixelMappingDetailCustomization_FixtureGroup::GetSelectedDMXLibrary(UDMXPixelMappingFixtureGroupComponent* FixtureGroupComponent) const
 {
 	if (FixtureGroupComponent)
diff --git a/Engine/Plugins/VirtualProduction/DMX/DMXPixelMapping/Source/DMXPixelMappingEditor/Private/Customizations/DMXPixelMappingDetailCustomization_FixtureGroup.h b/Engine/Plugins/VirtualProduction/DMX/DMXPixelMapping/Source/DMXPixelMappingEditor/Private/Customizations/DMXPixelMappingDetailCustomization_FixtureGroup.h
--- a/Engine/Plugins/VirtualProduction/DMX/DMXPixelMapping/Source/DMXPixelMappingEditor/Private/Customizations/DMXPixelMappingDetailCustomization_FixtureGroup.h
+++ b/Engine/Plugins/VirtualProduction/DMX/DMXPixelMapping/Source/DMXPixelMappingEditor/Private/Customizations/DMXPixelMappingDetailCustomization_FixtureGroup.h
@@ -43,6 +43,15 @@ public:
 	virtual void CustomizeDetails(IDetailLayoutBuilder& DetailLayout) override;
 
 private:
+	/** Creates a button that adds all patches to the group */
+	void CreateAddAllPatchesButton(IDetailLayoutBuilder& InDetailLayout);
+
+	/** Creates a list of fixture patches available to drag drop as detail rws */
+	void CreateFixturePatchDetailRows(IDetailLayoutBuilder& InDetailLayout);
+
+	/** Called when the 'Add all Patches' button was clicked */
+	FReply OnAddAllPatchesClicked();
+
 	/** Called when the library changed */
 	void OnLibraryChanged();
 
@@ -52,6 +61,15 @@ private:
 	/** Called when a component was removed */
 	void OnComponentRemoved(UDMXPixelMapping* PixelMapping, UDMXPixelMappingBaseComponent* Component);
 
+	/** Called before the SizeX property changed */
+	void OnSizePropertyPreChange();
+
+	/** Called when the SizeX property changed */
+	void OnSizePropertyChanged(const FPropertyChangedEvent& PropertyChangedEvent);
+
+	/** Handles the size property changed, useful to call on tick */
+	void HandleSizePropertyChanged();
+
 	/** Forces the detail layout to refresh */
 	void ForceRefresh();
 
@@ -64,12 +82,18 @@ private:
 	/** Called when a fixture patch was dragged */
 	FReply OnFixturePatchesDragged(const FGeometry& MyGeometry, const FPointerEvent& MouseEvent);
 
+	/** Returns true if the fixture patch is used in the pixel mapping */
+	bool IsFixturePatchAssignedToPixelMapping(UDMXEntityFixturePatch* FixturePatch) const;
+
 	/** Updates highlights for the fixture patches (selection) */
 	void UpdateFixturePatchHighlights();
 
 	/** Update fixture patches in use */
 	void UpdateFixturePatchesInUse(UDMXLibrary* DMXLibrary);
 
+	/** Updates the bCachedScaleChildrenWithParent member from UDMXPixelMappingLayoutSettings */
+	void UpdateCachedScaleChildrenWithParent();
+
 	/** Helper that returns the library selected in for the group */
 	UDMXLibrary* GetSelectedDMXLibrary(UDMXPixelMappingFixtureGroupComponent* FixtureGroupComponent) const;
 
@@ -88,12 +112,27 @@ private:
 	/** The single fixture group component in use */
 	TWeakObjectPtr<UDMXPixelMappingFixtureGroupComponent> WeakFixtureGroupComponent;
 
+	/** Number of fixture patch rows displayed */
+	int32 NumFixturePatchRows = 0;
+
+	/** SizeX before it got changed */
+	TMap<TWeakObjectPtr<UDMXPixelMappingFixtureGroupComponent>, FVector2D> PreEditChangeComponentToSizeMap;
+
+	/** Cached UDMXPixelMappingLayoutSettings::bScaleChildrenWithParen, to avoid repetitive reads during interactive changes */
+	bool bCachedScaleChildrenWithParent = false;
+
 	/** Handle to the dmx library property */
 	TSharedPtr<IPropertyHandle> DMXLibraryHandle;
 
 	/** Handle to the dmx library's entity array */
 	TSharedPtr<IPropertyHandle> EntitiesHandle;
 
+	/** Property handle for the SizeX property */
+	TSharedPtr<IPropertyHandle> SizeXHandle;
+
+	/** Property handle for the SizeY property */
+	TSharedPtr<IPropertyHandle> SizeYHandle;
+
 	/** Patches with their details row. Helps to multiselect and toggle highlights */
 	struct FDetailRowWidgetWithPatch
 	{
@@ -106,7 +145,4 @@ private:
 	FDelegateHandle RequestForceRefreshHandle;
 
 	TSharedPtr<IPropertyUtilities> PropertyUtilities;
-
-	/** If true, is refreshing */
-	bool bRefreshing = false;
 };
diff --git a/Engine/Plugins/VirtualProduction/DMX/DMXPixelMapping/Source/DMXPixelMappingEditor/Private/Customizations/DMXPixelMappingDetailCustomization_FixtureGroupItem.cpp b/Engine/Plugins/VirtualProduction/DMX/DMXPixelMapping/Source/DMXPixelMappingEditor/Private/Customizations/DMXPixelMappingDetailCustomization_FixtureGroupItem.cpp
--- a/Engine/Plugins/VirtualProduction/DMX/DMXPixelMapping/Source/DMXPixelMappingEditor/Private/Customizations/DMXPixelMappingDetailCustomization_FixtureGroupItem.cpp
+++ b/Engine/Plugins/VirtualProduction/DMX/DMXPixelMapping/Source/DMXPixelMappingEditor/Private/Customizations/DMXPixelMappingDetailCustomization_FixtureGroupItem.cpp
@@ -38,14 +38,13 @@ void FDMXPixelMappingDetailCustomization_FixtureGroupItem::CustomizeDetails(IDet
 	IDetailCategoryBuilder& OutputSettingsCategory = DetailLayout->EditCategory("Output Settings", FText::GetEmpty(), ECategoryPriority::Important);
 
 	// Hide absolute postition property handles
-	// Hide absolute postition property handles
-	TSharedPtr<IPropertyHandle> PositionXPropertyHandle = InDetailLayout.GetProperty(GET_MEMBER_NAME_CHECKED(UDMXPixelMappingOutputComponent, PositionX));
+	TSharedPtr<IPropertyHandle> PositionXPropertyHandle = InDetailLayout.GetProperty(UDMXPixelMappingOutputComponent::GetPositionXPropertyName());
 	InDetailLayout.HideProperty(PositionXPropertyHandle);
-	TSharedPtr<IPropertyHandle> PositionYPropertyHandle = InDetailLayout.GetProperty(GET_MEMBER_NAME_CHECKED(UDMXPixelMappingOutputComponent, PositionY));
+	TSharedPtr<IPropertyHandle> PositionYPropertyHandle = InDetailLayout.GetProperty(UDMXPixelMappingOutputComponent::GetPositionYPropertyName());
 	InDetailLayout.HideProperty(PositionYPropertyHandle);
-	TSharedPtr<IPropertyHandle> SizeXPropertyHandle = InDetailLayout.GetProperty(GET_MEMBER_NAME_CHECKED(UDMXPixelMappingOutputComponent, SizeX));
+	TSharedPtr<IPropertyHandle> SizeXPropertyHandle = InDetailLayout.GetProperty(UDMXPixelMappingOutputComponent::GetSizeXPropertyName());
 	InDetailLayout.HideProperty(SizeXPropertyHandle);
-	TSharedPtr<IPropertyHandle> SizeYPropertyHandle = InDetailLayout.GetProperty(GET_MEMBER_NAME_CHECKED(UDMXPixelMappingOutputComponent, SizeY));
+	TSharedPtr<IPropertyHandle> SizeYPropertyHandle = InDetailLayout.GetProperty(UDMXPixelMappingOutputComponent::GetSizeYPropertyName());
 	InDetailLayout.HideProperty(SizeXPropertyHandle);
 
 	// Add Function and ColorMode properties at the beginning
diff --git a/Engine/Plugins/VirtualProduction/DMX/DMXPixelMapping/Source/DMXPixelMappingEditor/Private/Customizations/DMXPixelMappingDetailCustomization_MatrixCell.cpp b/Engine/Plugins/VirtualProduction/DMX/DMXPixelMapping/Source/DMXPixelMappingEditor/Private/Customizations/DMXPixelMappingDetailCustomization_MatrixCell.cpp
--- a/Engine/Plugins/VirtualProduction/DMX/DMXPixelMapping/Source/DMXPixelMappingEditor/Private/Customizations/DMXPixelMappingDetailCustomization_MatrixCell.cpp
+++ b/Engine/Plugins/VirtualProduction/DMX/DMXPixelMapping/Source/DMXPixelMappingEditor/Private/Customizations/DMXPixelMappingDetailCustomization_MatrixCell.cpp
@@ -11,12 +11,12 @@
 void FDMXPixelMappingDetailCustomization_MatrixCell::CustomizeDetails(IDetailLayoutBuilder& InDetailLayout)
 {
 	// Hide absolute postition property handles
-	TSharedPtr<IPropertyHandle> PositionXPropertyHandle = InDetailLayout.GetProperty(GET_MEMBER_NAME_CHECKED(UDMXPixelMappingOutputComponent, PositionX));
-	TSharedPtr<IPropertyHandle> PositionYPropertyHandle = InDetailLayout.GetProperty(GET_MEMBER_NAME_CHECKED(UDMXPixelMappingOutputComponent, PositionY));
+	TSharedPtr<IPropertyHandle> PositionXPropertyHandle = InDetailLayout.GetProperty(UDMXPixelMappingOutputComponent::GetPositionXPropertyName());
+	TSharedPtr<IPropertyHandle> PositionYPropertyHandle = InDetailLayout.GetProperty(UDMXPixelMappingOutputComponent::GetPositionXPropertyName());
 	InDetailLayout.HideProperty(PositionXPropertyHandle);
 
 	// Hide size property handles
-	TSharedPtr<IPropertyHandle> SizeXPropertyHandle = InDetailLayout.GetProperty(GET_MEMBER_NAME_CHECKED(UDMXPixelMappingOutputComponent, SizeX));
-	TSharedPtr<IPropertyHandle> SizeYPropertyHandle = InDetailLayout.GetProperty(GET_MEMBER_NAME_CHECKED(UDMXPixelMappingOutputComponent, SizeY));
+	TSharedPtr<IPropertyHandle> SizeXPropertyHandle = InDetailLayout.GetProperty(UDMXPixelMappingOutputComponent::GetSizeXPropertyName());
+	TSharedPtr<IPropertyHandle> SizeYPropertyHandle = InDetailLayout.GetProperty(UDMXPixelMappingOutputComponent::GetSizeYPropertyName());
 	InDetailLayout.HideProperty(SizeXPropertyHandle);
 }
diff --git a/Engine/Plugins/VirtualProduction/DMX/DMXPixelMapping/Source/DMXPixelMappingEditor/Private/Customizations/DMXPixelMappingDetailCustomization_Renderer.cpp b/Engine/Plugins/VirtualProduction/DMX/DMXPixelMapping/Source/DMXPixelMappingEditor/Private/Customizations/DMXPixelMappingDetailCustomization_Renderer.cpp
--- a/Engine/Plugins/VirtualProduction/DMX/DMXPixelMapping/Source/DMXPixelMappingEditor/Private/Customizations/DMXPixelMappingDetailCustomization_Renderer.cpp
+++ b/Engine/Plugins/VirtualProduction/DMX/DMXPixelMapping/Source/DMXPixelMappingEditor/Private/Customizations/DMXPixelMappingDetailCustomization_Renderer.cpp
@@ -38,9 +38,9 @@ void FDMXPixelMappingDetailCustomization_Renderer::CustomizeDetails(IDetailLayou
 	FSimpleDelegate OnInputTexturePropertyChangedDelegate = FSimpleDelegate::CreateSP(this, &FDMXPixelMappingDetailCustomization_Renderer::OnInputTexturePropertyChanged);
 	RendererTypePropertyHandle->SetOnPropertyValueChanged(OnInputTexturePropertyChangedDelegate);
 
-	// Hide properties
-	TSharedRef<IPropertyHandle> PositionXPropertyHandle = DetailLayout.GetProperty(GET_MEMBER_NAME_CHECKED(UDMXPixelMappingOutputComponent, PositionX), UDMXPixelMappingOutputComponent::StaticClass());
-	TSharedRef<IPropertyHandle> PositionYPropertyHandle = DetailLayout.GetProperty(GET_MEMBER_NAME_CHECKED(UDMXPixelMappingOutputComponent, PositionY), UDMXPixelMappingOutputComponent::StaticClass());
+	// Hide Output Component properties
+	TSharedRef<IPropertyHandle> PositionXPropertyHandle = DetailLayout.GetProperty(UDMXPixelMappingOutputComponent::GetPositionXPropertyName(), UDMXPixelMappingOutputComponent::StaticClass());
+	TSharedRef<IPropertyHandle> PositionYPropertyHandle = DetailLayout.GetProperty(UDMXPixelMappingOutputComponent::GetPositionYPropertyName(), UDMXPixelMappingOutputComponent::StaticClass());
 	TSharedRef<IPropertyHandle> IsLockInDesignerPropertyHandle = DetailLayout.GetProperty(UDMXPixelMappingOutputComponent::GetLockInDesignerPropertyName(), UDMXPixelMappingOutputComponent::StaticClass());
 	TSharedRef<IPropertyHandle> IsVisibleInDesignerPropertyHandle = DetailLayout.GetProperty(UDMXPixelMappingOutputComponent::GetVisibleInDesignerPropertyName(), UDMXPixelMappingOutputComponent::StaticClass());
 	TSharedRef<IPropertyHandle> PixelBlendingPropertyHandle = DetailLayout.GetProperty(GET_MEMBER_NAME_CHECKED(UDMXPixelMappingOutputComponent, CellBlendingQuality), UDMXPixelMappingOutputComponent::StaticClass());
diff --git a/Engine/Plugins/VirtualProduction/DMX/DMXPixelMapping/Source/DMXPixelMappingEditor/Private/DMXPixelMappingEditorCommands.cpp b/Engine/Plugins/VirtualProduction/DMX/DMXPixelMapping/Source/DMXPixelMappingEditor/Private/DMXPixelMappingEditorCommands.cpp
--- a/Engine/Plugins/VirtualProduction/DMX/DMXPixelMapping/Source/DMXPixelMappingEditor/Private/DMXPixelMappingEditorCommands.cpp
+++ b/Engine/Plugins/VirtualProduction/DMX/DMXPixelMapping/Source/DMXPixelMappingEditor/Private/DMXPixelMappingEditorCommands.cpp
@@ -26,8 +26,16 @@ void FDMXPixelMappingEditorCommands::RegisterCommands()
 
 	UI_COMMAND(PlayDMX, "Play DMX", "Play DMX.", EUserInterfaceActionType::Button, FInputChord());
 	UI_COMMAND(StopPlayingDMX, "Stop Playing DMX", "Stop Playing DMX.", EUserInterfaceActionType::Button, FInputChord());
-
-	UI_COMMAND(bTogglePlayDMXAll, "Toggle Play DMX All", "Toggles playing all render components and all their childs", EUserInterfaceActionType::ToggleButton, FInputChord());
+	UI_COMMAND(TogglePlayDMXAll, "Toggle Play DMX All", "Toggles playing all render components and all their childs", EUserInterfaceActionType::ToggleButton, FInputChord());
+
+	UI_COMMAND(SizeComponentToTexture, "Size Component to Texture", "Sizes the selected group to Texture.", EUserInterfaceActionType::Button, FInputChord(EModifierKey::Shift, EKeys::S));
+	
+	UI_COMMAND(ToggleScaleChildrenWithParent, "Scale Children with Parent", "Sets if children are scaled with their parent, when the parent is resized.", EUserInterfaceActionType::ToggleButton, FInputChord(EModifierKey::Alt, EKeys::Q));
+	UI_COMMAND(ToggleAlwaysSelectGroup, "Always select Group", "Sets if the parent is selected, if a child is clicked.", EUserInterfaceActionType::ToggleButton, FInputChord(EModifierKey::Alt, EKeys::W));
+	UI_COMMAND(ToggleApplyLayoutScriptWhenLoaded, "Apply Layout Script instantly", "Sets if layout script are applied as soon as they are loaded.", EUserInterfaceActionType::ToggleButton, FInputChord(EModifierKey::Alt, EKeys::E));
+	UI_COMMAND(ToggleShowComponentNames, "Show Component Names", "Sets if the name of components are displayed. Note, may impact editor performance.", EUserInterfaceActionType::ToggleButton, FInputChord(EModifierKey::Alt, EKeys::A));
+	UI_COMMAND(ToggleShowPatchInfo, "Show Patch Info", "Sets if information about the pach is displayed. Note, may impact editor performance.", EUserInterfaceActionType::ToggleButton, FInputChord(EModifierKey::Alt, EKeys::S));
+	UI_COMMAND(ToggleShowCellIDs, "Show Cell IDs", "Sets if the cell IDs of matrix cells are displayed.", EUserInterfaceActionType::ToggleButton, FInputChord(EModifierKey::Alt, EKeys::D));
 }
 
 #undef LOCTEXT_NAMESPACE 
diff --git a/Engine/Plugins/VirtualProduction/DMX/DMXPixelMapping/Source/DMXPixelMappingEditor/Private/DMXPixelMappingEditorCommands.h b/Engine/Plugins/VirtualProduction/DMX/DMXPixelMapping/Source/DMXPixelMappingEditor/Private/DMXPixelMappingEditorCommands.h
--- a/Engine/Plugins/VirtualProduction/DMX/DMXPixelMapping/Source/DMXPixelMappingEditor/Private/DMXPixelMappingEditorCommands.h
+++ b/Engine/Plugins/VirtualProduction/DMX/DMXPixelMapping/Source/DMXPixelMappingEditor/Private/DMXPixelMappingEditorCommands.h
@@ -18,5 +18,15 @@ public:
 	TSharedPtr<FUICommandInfo> AddMapping;
 	TSharedPtr<FUICommandInfo> PlayDMX;
 	TSharedPtr<FUICommandInfo> StopPlayingDMX;
-	TSharedPtr<FUICommandInfo> bTogglePlayDMXAll;
-};
\ No newline at end of file
+	TSharedPtr<FUICommandInfo> TogglePlayDMXAll;
+
+	// Layout related
+	TSharedPtr<FUICommandInfo> SizeComponentToTexture;
+
+	TSharedPtr<FUICommandInfo> ToggleScaleChildrenWithParent;
+	TSharedPtr<FUICommandInfo> ToggleAlwaysSelectGroup;
+	TSharedPtr<FUICommandInfo> ToggleApplyLayoutScriptWhenLoaded;
+	TSharedPtr<FUICommandInfo> ToggleShowComponentNames;
+	TSharedPtr<FUICommandInfo> ToggleShowPatchInfo;
+	TSharedPtr<FUICommandInfo> ToggleShowCellIDs;
+};
diff --git a/Engine/Plugins/VirtualProduction/DMX/DMXPixelMapping/Source/DMXPixelMappingEditor/Private/DMXPixelMappingLayoutSettings.cpp b/Engine/Plugins/VirtualProduction/DMX/DMXPixelMapping/Source/DMXPixelMappingEditor/Private/DMXPixelMappingLayoutSettings.cpp
new file mode 100644
--- /dev/null
+++ b/Engine/Plugins/VirtualProduction/DMX/DMXPixelMapping/Source/DMXPixelMappingEditor/Private/DMXPixelMappingLayoutSettings.cpp
@@ -0,0 +1,17 @@
+// Copyright Epic Games, Inc. All Rights Reserved.
+
+#include "DMXPixelMappingLayoutSettings.h"
+
+
+FSimpleMulticastDelegate UDMXPixelMappingLayoutSettings::OnLayoutSettingsChanged;
+
+void UDMXPixelMappingLayoutSettings::PostEditChangeProperty(FPropertyChangedEvent& PropertyChangedEvent)
+{
+	Super::PostEditChangeProperty(PropertyChangedEvent);
+
+	if (PropertyChangedEvent.ChangeType != EPropertyChangeType::Interactive)
+	{
+		SaveConfig();
+		OnLayoutSettingsChanged.Broadcast();
+	}
+}
diff --git a/Engine/Plugins/VirtualProduction/DMX/DMXPixelMapping/Source/DMXPixelMappingEditor/Private/DMXPixelMappingLayoutSettings.h b/Engine/Plugins/VirtualProduction/DMX/DMXPixelMapping/Source/DMXPixelMappingEditor/Private/DMXPixelMappingLayoutSettings.h
new file mode 100644
--- /dev/null
+++ b/Engine/Plugins/VirtualProduction/DMX/DMXPixelMapping/Source/DMXPixelMappingEditor/Private/DMXPixelMappingLayoutSettings.h
@@ -0,0 +1,54 @@
+// Copyright Epic Games, Inc. All Rights Reserved.
+
+#pragma once
+
+#include "DMXPixelMappingLayoutSettings.generated.h"
+
+
+/** Layout options for the Pixel Mapping editor */
+UCLASS(Config = DMXPixelMappingEditor, DefaultConfig, meta = (DisplayName = "DMXEditor"))
+class UDMXPixelMappingLayoutSettings 
+	: public UObject
+{
+	GENERATED_BODY()
+
+public:
+	/** If true, scales children when the parent component is resized */
+	UPROPERTY(Config, EditAnywhere, Category = "Layout Settings")
+	bool bScaleChildrenWithParent = true;
+
+	/** If true, selects parent when a child is clicked */
+	UPROPERTY(Config, EditAnywhere, Category = "Layout Settings")
+	bool bAlwaysSelectGroup = false;
+
+	/** If true, applies layout scripts instantly when they are loaded */
+	UPROPERTY(Config, EditAnywhere, Category = "Layout Settings")
+	bool bApplyLayoutScriptWhenLoaded = true;
+
+	/** If true, shows the name of the Fixture Patch where applicable. May affect editor performance. */
+	UPROPERTY(Config, EditAnywhere, Category = "Layout Settings")
+	bool bShowComponentNames = true;
+
+	/** If true, shows Fixture info bout the Fixture Patch where applicable. May affect editor performance. */
+	UPROPERTY(Config, EditAnywhere, Category = "Layout Settings")
+	bool bShowPatchInfo = false;
+
+	/** If true, shows Cell IDs where applicable. May affect editor performance. */
+	UPROPERTY(Config, EditAnywhere, Category = "Layout Settings")
+	bool bShowCellIDs = false;
+
+	/** Gets a delegate broadcast when Layout Settings changed */
+	static FSimpleMulticastDelegate& GetOnLayoutSettingsChanged() { return OnLayoutSettingsChanged; }
+
+protected:
+	//~ Being UObject interface
+#if WITH_EDITOR
+	virtual void PostEditChangeProperty(FPropertyChangedEvent& PropertyChangedEvent) override;
+#endif
+	//~ End UObject interface
+
+private:
+	/** Broadcast when Layout Settings changed */
+	static FSimpleMulticastDelegate OnLayoutSettingsChanged;
+
+};
diff --git a/Engine/Plugins/VirtualProduction/DMX/DMXPixelMapping/Source/DMXPixelMappingEditor/Private/DMXPixelMappingToolbar.cpp b/Engine/Plugins/VirtualProduction/DMX/DMXPixelMapping/Source/DMXPixelMappingEditor/Private/DMXPixelMappingToolbar.cpp
--- a/Engine/Plugins/VirtualProduction/DMX/DMXPixelMapping/Source/DMXPixelMappingEditor/Private/DMXPixelMappingToolbar.cpp
+++ b/Engine/Plugins/VirtualProduction/DMX/DMXPixelMapping/Source/DMXPixelMappingEditor/Private/DMXPixelMappingToolbar.cpp
@@ -24,17 +24,11 @@ void FDMXPixelMappingToolbar::BuildToolbar(TSharedPtr<FExtender> Extender)
 		"Asset",
 		EExtensionHook::After,
 		Toolkit->GetToolkitCommands(),
-		FToolBarExtensionDelegate::CreateSP(this, &FDMXPixelMappingToolbar::Build)
+		FToolBarExtensionDelegate::CreateSP(this, &FDMXPixelMappingToolbar::BuildToolbarCallback)
 	);
 }
 
-void FDMXPixelMappingToolbar::Build(FToolBarBuilder& ToolbarBuilder)
-{
-	AddHelpersSection(ToolbarBuilder);
-	AddPlayAndStopSection(ToolbarBuilder);
-}
-
-void FDMXPixelMappingToolbar::AddHelpersSection(FToolBarBuilder& ToolbarBuilder)
+void FDMXPixelMappingToolbar::BuildToolbarCallback(FToolBarBuilder& ToolbarBuilder)
 {
 	ToolbarBuilder.BeginSection("Thumbnail");
 	{
@@ -44,10 +38,7 @@ void FDMXPixelMappingToolbar::AddHelpersSection(FToolBarBuilder& ToolbarBuilder)
 			FSlateIcon(FAppStyle::GetAppStyleSetName(), "Cascade.SaveThumbnailImage"));
 	}
 	ToolbarBuilder.EndSection();
-}
 
-void FDMXPixelMappingToolbar::AddPlayAndStopSection(FToolBarBuilder& ToolbarBuilder)
-{
 	TSharedPtr<FDMXPixelMappingToolkit> Toolkit = ToolkitWeakPtr.Pin();
 	check(Toolkit.IsValid());
 
@@ -81,16 +72,28 @@ void FDMXPixelMappingToolbar::AddPlayAndStopSection(FToolBarBuilder& ToolbarBuil
 
 		ToolbarBuilder.AddComboButton(
 			PlayDMXOptionsAction,
-			FOnGetContent::CreateSP(this, &FDMXPixelMappingToolbar::FillPlayMenu),
+			FOnGetContent::CreateSP(this, &FDMXPixelMappingToolbar::GeneratesPlayOptionsWidget),
 			LOCTEXT("PlayDMXOptions", "Play DMX Options"),
 			LOCTEXT("PlayDMXOptions", "Play DMX Options"),
 			FSlateIcon(),
 			true);
 	}
 	ToolbarBuilder.EndSection();
+
+	ToolbarBuilder.BeginSection("Layout");
+	{
+		ToolbarBuilder.AddComboButton(
+			FUIAction(),
+			FOnGetContent::CreateSP(this, &FDMXPixelMappingToolbar::GenerateLayoutMenu),
+			LOCTEXT("PixelMappingToolbarLayoutSettingsLabel", "Layout"),
+			LOCTEXT("PixelMappingToolbarLayoutSettingsTooltip", "Layout Settings for Pixel Mapping Editors"),
+			FSlateIcon(FAppStyle::GetAppStyleSetName(), "Icons.Layout")
+		);
+	}
+	ToolbarBuilder.EndSection();
 }
 
-TSharedRef<SWidget> FDMXPixelMappingToolbar::FillPlayMenu()
+TSharedRef<SWidget> FDMXPixelMappingToolbar::GeneratesPlayOptionsWidget()
 {
 	TSharedPtr<FDMXPixelMappingToolkit> Toolkit = ToolkitWeakPtr.Pin();
 	check(Toolkit.IsValid());
@@ -101,7 +104,7 @@ TSharedRef<SWidget> FDMXPixelMappingToolbar::FillPlayMenu()
 	{
 		MenuBuilder.BeginSection("bTogglePlayDMXAll");
 		{
-			MenuBuilder.AddMenuEntry(FDMXPixelMappingEditorCommands::Get().bTogglePlayDMXAll);
+			MenuBuilder.AddMenuEntry(FDMXPixelMappingEditorCommands::Get().TogglePlayDMXAll);
 		}
 		MenuBuilder.EndSection();
 	}
@@ -109,4 +112,32 @@ TSharedRef<SWidget> FDMXPixelMappingToolbar::FillPlayMenu()
 	return MenuBuilder.MakeWidget();
 }
 
-#undef LOCTEXT_NAMESPACE
\ No newline at end of file
+TSharedRef<SWidget> FDMXPixelMappingToolbar::GenerateLayoutMenu()
+{
+	TSharedPtr<FDMXPixelMappingToolkit> Toolkit = ToolkitWeakPtr.Pin();
+	check(Toolkit.IsValid());
+
+	constexpr bool bShouldCloseMenuAfterSelection = true;
+	FMenuBuilder MenuBuilder(bShouldCloseMenuAfterSelection, Toolkit->GetToolkitCommands());
+	
+	MenuBuilder.BeginSection("Actions", LOCTEXT("ActionsSection", "Actions"));
+	{
+		MenuBuilder.AddMenuEntry(FDMXPixelMappingEditorCommands::Get().SizeComponentToTexture);
+	}
+	MenuBuilder.EndSection();
+
+	MenuBuilder.BeginSection("LayoutSettings", LOCTEXT("LayoutSettingsSection", "Settings"));
+	{
+		MenuBuilder.AddMenuEntry(FDMXPixelMappingEditorCommands::Get().ToggleScaleChildrenWithParent);
+		MenuBuilder.AddMenuEntry(FDMXPixelMappingEditorCommands::Get().ToggleAlwaysSelectGroup);
+		MenuBuilder.AddMenuEntry(FDMXPixelMappingEditorCommands::Get().ToggleApplyLayoutScriptWhenLoaded);
+		MenuBuilder.AddMenuEntry(FDMXPixelMappingEditorCommands::Get().ToggleShowComponentNames);
+		MenuBuilder.AddMenuEntry(FDMXPixelMappingEditorCommands::Get().ToggleShowPatchInfo);
+		MenuBuilder.AddMenuEntry(FDMXPixelMappingEditorCommands::Get().ToggleShowCellIDs);
+	}
+	MenuBuilder.EndSection();
+
+	return MenuBuilder.MakeWidget();
+}
+
+#undef LOCTEXT_NAMESPACE
diff --git a/Engine/Plugins/VirtualProduction/DMX/DMXPixelMapping/Source/DMXPixelMappingEditor/Private/DMXPixelMappingToolbar.h b/Engine/Plugins/VirtualProduction/DMX/DMXPixelMapping/Source/DMXPixelMappingEditor/Private/DMXPixelMappingToolbar.h
--- a/Engine/Plugins/VirtualProduction/DMX/DMXPixelMapping/Source/DMXPixelMappingEditor/Private/DMXPixelMappingToolbar.h
+++ b/Engine/Plugins/VirtualProduction/DMX/DMXPixelMapping/Source/DMXPixelMappingEditor/Private/DMXPixelMappingToolbar.h
@@ -22,14 +22,18 @@ public:
 	/** Virtual Destructor */
 	virtual ~FDMXPixelMappingToolbar() {}
 
+	/** Builds the toolbar */
 	void BuildToolbar(TSharedPtr<FExtender> Extender);
 
 private:
-	void Build(FToolBarBuilder& ToolbarBuilder);
+	/** Callback, raised when the menu extender requests to build the toolbar */
+	void BuildToolbarCallback(FToolBarBuilder& ToolbarBuilder);
 
-	void AddHelpersSection(FToolBarBuilder& ToolbarBuilder);
-	void AddPlayAndStopSection(FToolBarBuilder& ToolbarBuilder);
-	TSharedRef<SWidget> FillPlayMenu();
+	/** Generates a widget with play options (play, stop) */
+	TSharedRef<SWidget> GeneratesPlayOptionsWidget();
+
+	/** Generates a layout settings menu for the pixelmapping toolbar */
+	TSharedRef<SWidget> GenerateLayoutMenu();
 
 public:
 	TWeakPtr<FDMXPixelMappingToolkit> ToolkitWeakPtr;
diff --git a/Engine/Plugins/VirtualProduction/DMX/DMXPixelMapping/Source/DMXPixelMappingEditor/Private/DragDrop/DMXPixelMappingDragDropOp.cpp b/Engine/Plugins/VirtualProduction/DMX/DMXPixelMapping/Source/DMXPixelMappingEditor/Private/DragDrop/DMXPixelMappingDragDropOp.cpp
--- a/Engine/Plugins/VirtualProduction/DMX/DMXPixelMapping/Source/DMXPixelMappingEditor/Private/DragDrop/DMXPixelMappingDragDropOp.cpp
+++ b/Engine/Plugins/VirtualProduction/DMX/DMXPixelMapping/Source/DMXPixelMappingEditor/Private/DragDrop/DMXPixelMappingDragDropOp.cpp
@@ -21,37 +21,17 @@
 
 FDMXPixelMappingDragDropOp::~FDMXPixelMappingDragDropOp()
 {
-	// Remember children not over their parent (to remove after), make others highest ZOrder.
+	// Newly added components are highest Z-Order
 	TArray<UDMXPixelMappingOutputComponent*> ChildrenNotOverTheirParent;
 	for (const TWeakObjectPtr<UDMXPixelMappingBaseComponent>& DraggedComponent : DraggedComponents)
 	{
 		if (UDMXPixelMappingOutputComponent* DraggedOutputComponent = Cast<UDMXPixelMappingOutputComponent>(DraggedComponent.Get()))
 		{
-			if (DraggedOutputComponent->GetParent() && 
-				!DraggedOutputComponent->IsOverParent())
-			{
-				ChildrenNotOverTheirParent.Add(DraggedOutputComponent);
-			}
-			else
-			{
-				DraggedOutputComponent->MakeHighestZOrderInComponentRect();
-			}
+			DraggedOutputComponent->MakeHighestZOrderInComponentRect();
 		}
 	}
 
-	// Destroy components not over their parent 
-	TArray<UDMXPixelMappingOutputComponent*> RemovedComponents;
-	for (UDMXPixelMappingOutputComponent* ComponentNotOverParent : ChildrenNotOverTheirParent)
-	{
-		DraggedComponents.Remove(ComponentNotOverParent);
-
-		// Remove output components not over their parent
-		ComponentNotOverParent->Modify();
-		ComponentNotOverParent->GetParent()->Modify();
-
-		ComponentNotOverParent->GetParent()->RemoveChild(ComponentNotOverParent);
-	}
-
+	// End the transaction in any case
 	GEditor->EndTransaction();
 }
 
diff --git a/Engine/Plugins/VirtualProduction/DMX/DMXPixelMapping/Source/DMXPixelMappingEditor/Private/DragDrop/DMXPixelMappingGroupChildDragDropHelper.cpp b/Engine/Plugins/VirtualProduction/DMX/DMXPixelMapping/Source/DMXPixelMappingEditor/Private/DragDrop/DMXPixelMappingGroupChildDragDropHelper.cpp
--- a/Engine/Plugins/VirtualProduction/DMX/DMXPixelMapping/Source/DMXPixelMappingEditor/Private/DragDrop/DMXPixelMappingGroupChildDragDropHelper.cpp
+++ b/Engine/Plugins/VirtualProduction/DMX/DMXPixelMapping/Source/DMXPixelMappingEditor/Private/DragDrop/DMXPixelMappingGroupChildDragDropHelper.cpp
@@ -62,17 +62,6 @@ TSharedPtr<FDMXPixelMappingGroupChildDragDropHelper> FDMXPixelMappingGroupChildD
 		}
 
 		NewHelper->WeakChildComponents.Add(OutputComponent);
-
-		// Clamp the size of each child to the a valid size
-		FVector2D IndividualChildSize = OutputComponent->GetSize();
-		if (IndividualChildSize.ComponentwiseAllGreaterThan(NewHelper->ParentSize))
-		{
-			OutputComponent->SetSize(NewHelper->ParentSize - FVector2D(NewHelper->ParentSize.X / 10.f, NewHelper->ParentSize.X / 10.f)); // Same X and Y so it's evenly aligned
-		}
-		else if (IndividualChildSize.GetMax() < 1.f)
-		{
-			OutputComponent->SetSize(FVector2D(1.f, 1.f));
-		}
 	}
 
 	return NewHelper;
@@ -106,7 +95,7 @@ void FDMXPixelMappingGroupChildDragDropHelper::LayoutAligned(const FVector2D& Gr
 					if (GroupComponent->IsOverPosition(NextPosition) && 
 						GroupComponent->IsOverPosition(NextPosition + ChildComponent->GetSize()))
 					{
-						SetPositionRounded(ChildComponent, NextPosition);
+						SetPosition(ChildComponent, NextPosition);
 
 						RowHeight = FMath::Max(ChildComponent->GetSize().Y, RowHeight);
 						NextPosition = FVector2D(NextPosition.X + ChildComponent->GetSize().X, NextPosition.Y);
@@ -120,14 +109,14 @@ void FDMXPixelMappingGroupChildDragDropHelper::LayoutAligned(const FVector2D& Gr
 						if (GroupComponent->IsOverPosition(NextPositionOnNewRow) &&
 							GroupComponent->IsOverPosition(NextPositionOnNewRow + ChildComponent->GetSize()))
 						{
-							SetPositionRounded(ChildComponent, NewRowPosition);
+							SetPosition(ChildComponent, NewRowPosition);
 
 							NextPosition = FVector2D(NewRowPosition.X + ChildComponent->GetSize().X, NewRowPosition.Y);
 							RowHeight = ChildComponent->GetSize().Y;							
 						}
 						else
 						{
-							SetPositionRounded(ChildComponent, NextPosition);
+							SetPosition(ChildComponent, NextPosition);
 
 							NextPosition = FVector2D(NextPosition.X + ChildComponent->GetSize().X, NextPosition.Y);
 						}
@@ -156,16 +145,7 @@ void FDMXPixelMappingGroupChildDragDropHelper::LayoutUnaligned(const FVector2D&
 						FVector2D AnchorOffset = Anchor - ChildComponent->GetPosition();
 
 						FVector2D NewPosition = GraphSpacePosition - AnchorOffset - PinnedDragDropOp->GraphSpaceDragOffset;
-						SetPositionRounded(ChildComponent, NewPosition);
-
-						// Apply a color depending on where the component ended
-						if (TSharedPtr<FDMXPixelMappingComponentWidget> ComponentWidget = ChildComponent->GetComponentWidget())
-						{
-							const bool bOverParent = ChildComponent->IsOverParent();
-							const FLinearColor WidgetColor = bOverParent ? FLinearColor::Green : FLinearColor::Red;
-
-							ComponentWidget->SetColor(WidgetColor);
-						}
+						SetPosition(ChildComponent, NewPosition);
 					}
 				}
 			}
@@ -173,12 +153,12 @@ void FDMXPixelMappingGroupChildDragDropHelper::LayoutUnaligned(const FVector2D&
 	}
 }
 
-void FDMXPixelMappingGroupChildDragDropHelper::SetPositionRounded(UDMXPixelMappingOutputComponent* Component, const FVector2D& Position) const
+void FDMXPixelMappingGroupChildDragDropHelper::SetPosition(UDMXPixelMappingOutputComponent* Component, const FVector2D& Position) const
 {
 	if (Component)
 	{
 		Component->Modify();
 
-		Component->SetPosition(Position.RoundToVector());
+		Component->SetPosition(Position);
 	}
 }
diff --git a/Engine/Plugins/VirtualProduction/DMX/DMXPixelMapping/Source/DMXPixelMappingEditor/Private/DragDrop/DMXPixelMappingGroupChildDragDropHelper.h b/Engine/Plugins/VirtualProduction/DMX/DMXPixelMapping/Source/DMXPixelMappingEditor/Private/DragDrop/DMXPixelMappingGroupChildDragDropHelper.h
--- a/Engine/Plugins/VirtualProduction/DMX/DMXPixelMapping/Source/DMXPixelMappingEditor/Private/DragDrop/DMXPixelMappingGroupChildDragDropHelper.h
+++ b/Engine/Plugins/VirtualProduction/DMX/DMXPixelMapping/Source/DMXPixelMappingEditor/Private/DragDrop/DMXPixelMappingGroupChildDragDropHelper.h
@@ -34,8 +34,8 @@ private:
 	/** Helper that lays out cells unaligned */
 	void LayoutUnaligned(const FVector2D& GraphSpacePosition);
 
-	/** Sets the position of the component rounded */
-	void SetPositionRounded(UDMXPixelMappingOutputComponent* Component, const FVector2D& Position) const;
+	/** Sets the position of the component */
+	void SetPosition(UDMXPixelMappingOutputComponent* Component, const FVector2D& Position) const;
 
 	/** Local position of the parent */
 	FVector2D ParentPosition;
diff --git a/Engine/Plugins/VirtualProduction/DMX/DMXPixelMapping/Source/DMXPixelMappingEditor/Private/Toolkits/DMXPixelMappingToolkit.cpp b/Engine/Plugins/VirtualProduction/DMX/DMXPixelMapping/Source/DMXPixelMappingEditor/Private/Toolkits/DMXPixelMappingToolkit.cpp
--- a/Engine/Plugins/VirtualProduction/DMX/DMXPixelMapping/Source/DMXPixelMappingEditor/Private/Toolkits/DMXPixelMappingToolkit.cpp
+++ b/Engine/Plugins/VirtualProduction/DMX/DMXPixelMapping/Source/DMXPixelMappingEditor/Private/Toolkits/DMXPixelMappingToolkit.cpp
@@ -3,6 +3,7 @@
 #include "Toolkits/DMXPixelMappingToolkit.h"
 
 #include "DMXEditorUtils.h"
+#include "DMXPixelMappingLayoutSettings.h"
 #include "DMXPixelMapping.h"
 #include "DMXPixelMappingEditorCommands.h"
 #include "DMXPixelMappingEditorModule.h"
@@ -17,26 +18,21 @@
 #include "Components/DMXPixelMappingRootComponent.h"
 #include "Components/DMXPixelMappingMatrixComponent.h"
 #include "Components/DMXPixelMappingMatrixCellComponent.h"
+#include "Components/DMXPixelMappingScreenComponent.h"
 #include "Templates/DMXPixelMappingComponentTemplate.h"
 #include "ViewModels/DMXPixelMappingPaletteViewModel.h"
-#include "Views/SDMXPixelMappingPaletteView.h"
-#include "Views/SDMXPixelMappingHierarchyView.h"
 #include "Views/SDMXPixelMappingDesignerView.h"
-#include "Views/SDMXPixelMappingPreviewView.h"
 #include "Views/SDMXPixelMappingDetailsView.h"
-
-#include "Library/DMXEntityFixtureType.h"
-#include "Library/DMXEntityFixturePatch.h"
-#include "Library/DMXLibrary.h"
-#include "DMXPixelMappingComponentWidget.h"
+#include "Views/SDMXPixelMappingHierarchyView.h"
+#include "Views/SDMXPixelMappingLayoutView.h"
+#include "Views/SDMXPixelMappingPaletteView.h"
+#include "Views/SDMXPixelMappingPreviewView.h"
 
 #include "ScopedTransaction.h"
 #include "Engine/TextureRenderTarget2D.h"
 #include "Framework/Commands/GenericCommands.h"
-#include "Kismet2/BlueprintEditorUtils.h"
 #include "Misc/ScopedSlowTask.h"
 #include "Styling/AppStyle.h"
-#include "Widgets/SWidget.h"
 #include "Widgets/Docking/SDockTab.h"
 
 
@@ -44,11 +40,10 @@
 
 const FName FDMXPixelMappingToolkit::PaletteViewTabID(TEXT("DMXPixelMappingEditor_PaletteViewTabID"));
 const FName FDMXPixelMappingToolkit::HierarchyViewTabID(TEXT("DMXPixelMappingEditor_HierarchyViewTabID"));
-
 const FName FDMXPixelMappingToolkit::DesignerViewTabID(TEXT("DMXPixelMappingEditor_DesignerViewTabID"));
 const FName FDMXPixelMappingToolkit::PreviewViewTabID(TEXT("DMXPixelMappingEditor_PreviewViewTabID"));
-
 const FName FDMXPixelMappingToolkit::DetailsViewTabID(TEXT("DMXPixelMappingEditor_DetailsViewTabID"));
+const FName FDMXPixelMappingToolkit::LayoutViewTabID(TEXT("DMXPixelMappingEditor_LayoutViewTabID"));
 
 const uint8 FDMXPixelMappingToolkit::RequestStopSendingMaxTicks = 5;
 
@@ -108,6 +103,10 @@ void FDMXPixelMappingToolkit::RegisterTabSpawners(const TSharedRef<class FTabMan
 		.SetDisplayName(LOCTEXT("Tab_DetailsView", "Details"))
 		.SetGroup(WorkspaceMenuCategoryRef)
 		.SetIcon(FSlateIcon(FAppStyle::Get().GetStyleSetName(), "Icons.Details"));
+
+	InTabManager->RegisterTabSpawner(LayoutViewTabID, FOnSpawnTab::CreateSP(this, &FDMXPixelMappingToolkit::SpawnTab_LayoutView))
+		.SetDisplayName(LOCTEXT("Tab_LayoutView", "Layout"))
+		.SetGroup(WorkspaceMenuCategoryRef);
 }
 
 void FDMXPixelMappingToolkit::UnregisterTabSpawners(const TSharedRef<class FTabManager>& InTabManager)
@@ -119,6 +118,7 @@ void FDMXPixelMappingToolkit::UnregisterTabSpawners(const TSharedRef<class FTabM
 	InTabManager->UnregisterTabSpawner(DesignerViewTabID);
 	InTabManager->UnregisterTabSpawner(PreviewViewTabID);
 	InTabManager->UnregisterTabSpawner(DetailsViewTabID);
+	InTabManager->UnregisterTabSpawner(LayoutViewTabID);
 }
 
 FText FDMXPixelMappingToolkit::GetBaseToolkitName() const
@@ -262,9 +262,6 @@ void FDMXPixelMappingToolkit::Tick(float DeltaTime)
 			}
 		}
 	}
-
-	// Update component widgets
-	UpdateComponentWidgetColors();
 }
 
 TStatId FDMXPixelMappingToolkit::GetStatId() const
@@ -282,43 +279,6 @@ void FDMXPixelMappingToolkit::SetActiveRenderComponent(UDMXPixelMappingRendererC
 	ActiveRendererComponent = InComponent;
 }
 
-void FDMXPixelMappingToolkit::HandleAddComponents()
-{
-	// DEPRECATED 4.27
-	UpdateBlueprintNodes(DMXPixelMapping);
-	OnComponentsAddedOrDeletedDelegate_DEPRECATED.Broadcast();
-}
-
-void FDMXPixelMappingToolkit::HandleRemoveComponents()
-{
-	// DEPRECATED 4.27
-	FDMXEditorUtils::ClearAllDMXPortBuffers();
-	FDMXEditorUtils::ClearFixturePatchCachedData();
-
-	UpdateBlueprintNodes(DMXPixelMapping);
-	OnComponentsAddedOrDeletedDelegate_DEPRECATED.Broadcast();
-
-	// Remove removed components from selection
-	TSet<FDMXPixelMappingComponentReference> NewSelection;
-	for (const FDMXPixelMappingComponentReference& ComponentReference : SelectedComponents)
-	{
-
-		if (UDMXPixelMappingBaseComponent* Component = ComponentReference.GetComponent())
-		{
-			if (IsValid(Component->GetParent()) || Component->GetClass() == UDMXPixelMappingRootComponent::StaticClass())
-			{
-				NewSelection.Add(ComponentReference);
-			}
-		}
-	}
-	if (NewSelection.Num() != SelectedComponents.Num())
-	{
-		SelectedComponents = NewSelection;
-
-		OnSelectedComponentsChangedDelegate.Broadcast();
-	}
-}
-
 template <typename ComponentType>
 TArray<ComponentType> FDMXPixelMappingToolkit::MakeComponentArray(const TSet<FDMXPixelMappingComponentReference>& Components) const
 {
@@ -365,107 +325,6 @@ void FDMXPixelMappingToolkit::SelectComponents(const TSet<FDMXPixelMappingCompon
 	OnSelectedComponentsChangedDelegate.Broadcast();
 }
 
-void FDMXPixelMappingToolkit::UpdateComponentWidgetColors()
-{
-	if (DMXPixelMapping)
-	{
-		// Reset all to default color
-		DMXPixelMapping->ForEachComponentOfClass<UDMXPixelMappingOutputComponent>([this](UDMXPixelMappingOutputComponent* OutputComponent)
-			{
-				if (TSharedPtr<FDMXPixelMappingComponentWidget> ComponentWidget = OutputComponent->GetComponentWidget())
-				{
-					ComponentWidget->SetColor(OutputComponent->GetEditorColor());
-				}
-			});
-
-		TArray<UDMXPixelMappingBaseComponent*> SelectedComponentsArray;
-		SelectedComponentsArray.Reserve(GetSelectedComponents().Num());
-		for (const FDMXPixelMappingComponentReference& ComponentReference : GetSelectedComponents())
-		{
-			if (ComponentReference.GetComponent())
-			{
-				SelectedComponentsArray.Add(ComponentReference.GetComponent());
-			}
-		}
-
-		UDMXPixelMappingBaseComponent** SelectedFixtureGroupPtr = SelectedComponentsArray.FindByPredicate([&SelectedComponentsArray](UDMXPixelMappingBaseComponent* Component)
-			{
-				if (Component->GetClass() == UDMXPixelMappingFixtureGroupComponent::StaticClass())
-				{
-					return true;
-				}
-				return false;
-			});
-
-		if (SelectedFixtureGroupPtr)
-		{
-			if (UDMXPixelMappingOutputComponent* FixtureGroupComponent = Cast<UDMXPixelMappingOutputComponent>(*SelectedFixtureGroupPtr))
-			{
-				if (TSharedPtr<FDMXPixelMappingComponentWidget> ComponentWidget = FixtureGroupComponent->GetComponentWidget())
-				{
-					ComponentWidget->SetColor(FLinearColor::Green);
-				}
-
-				constexpr bool bApplyColorChildsRecursive = true;
-				FixtureGroupComponent->ForEachChildOfClass<UDMXPixelMappingOutputComponent>([FixtureGroupComponent](UDMXPixelMappingOutputComponent* HighlightedComponent)
-					{
-						const bool bIsOverParent = HighlightedComponent->IsOverParent();
-						if (TSharedPtr<FDMXPixelMappingComponentWidget> ComponentWidget = HighlightedComponent->GetComponentWidget())
-						{
-							if (bIsOverParent)
-							{
-								// Give childs of the group some greenish color when they're over their parent
-								ComponentWidget->SetColor(FLinearColor(0.05f, 0.5f, 0.05f));
-							}
-							else
-							{
-								ComponentWidget->SetColor(FLinearColor::Red);
-							}
-						}
-					}, bApplyColorChildsRecursive);
-			}
-		}
-		else
-		{
-			// Set a selected color to selected components and their childs
-			for (UDMXPixelMappingBaseComponent* Component : SelectedComponentsArray)
-			{
-				if (UDMXPixelMappingOutputComponent* OutputComponent = Cast<UDMXPixelMappingOutputComponent>(Component))
-				{
-					// Pick a color depending on where the widget resides and what is selected
-					FLinearColor Color = [OutputComponent]()
-					{
-						const bool bIsOverParent = OutputComponent->IsOverParent();
-						if (bIsOverParent)
-						{
-							return FLinearColor::Green;
-						}
-						else
-						{
-							return FLinearColor::Red;
-						}
-					}();
-
-					if (TSharedPtr<FDMXPixelMappingComponentWidget> ComponentWidget = OutputComponent->GetComponentWidget())
-					{
-						ComponentWidget->SetColor(Color);
-					}
-
-					constexpr bool bApplyColorChildsRecursive = true;
-					OutputComponent->ForEachChildOfClass<UDMXPixelMappingOutputComponent>([&Color](UDMXPixelMappingOutputComponent* HighlightedComponent)
-						{
-
-							if (TSharedPtr<FDMXPixelMappingComponentWidget> ComponentWidget = HighlightedComponent->GetComponentWidget())
-							{
-								ComponentWidget->SetColor(Color);
-							}
-						}, bApplyColorChildsRecursive);
-				}
-			}
-		}
-	}
-}
-
 bool FDMXPixelMappingToolkit::IsComponentSelected(UDMXPixelMappingBaseComponent* Component) const
 {
 	for (const FDMXPixelMappingComponentReference& ComponentReference : SelectedComponents)
@@ -508,7 +367,103 @@ void FDMXPixelMappingToolkit::StopPlayingDMX()
 
 void FDMXPixelMappingToolkit::ExecutebTogglePlayDMXAll()
 {
-	bTogglePlayDMXAll ^= 1;
+	bTogglePlayDMXAll = !bTogglePlayDMXAll;
+}
+
+bool FDMXPixelMappingToolkit::CanSizeSelectedComponentToTexture() const
+{
+	if (SelectedComponents.Num() == 1)
+	{
+		if (UDMXPixelMappingBaseComponent* Component = SelectedComponents.Array()[0].GetComponent())
+		{
+			return
+				Component->GetClass() == UDMXPixelMappingFixtureGroupComponent::StaticClass() ||
+				Component->GetClass() == UDMXPixelMappingMatrixComponent::StaticClass() ||
+				Component->GetClass() == UDMXPixelMappingScreenComponent::StaticClass() ||
+				(Component->GetParent() && Component->GetParent()->GetClass() == UDMXPixelMappingFixtureGroupComponent::StaticClass()) ||
+				(Component->GetParent() && Component->GetParent()->GetClass() == UDMXPixelMappingMatrixComponent::StaticClass()) ||
+				(Component->GetParent() && Component->GetParent()->GetClass() == UDMXPixelMappingScreenComponent::StaticClass());
+		}
+	}
+	return false;
+}
+
+void FDMXPixelMappingToolkit::SizeSelectedComponentToTexture()
+{
+	if (SelectedComponents.IsEmpty())
+	{
+		return;
+	}	
+
+	const UDMXPixelMappingRendererComponent* RendererComponent = GetActiveRendererComponent();
+	if (!RendererComponent)
+	{
+		return;
+	}
+
+	const FVector2D TextureSize = RendererComponent->GetSize();
+	if (TextureSize == FVector2D::ZeroVector)
+	{
+		return;
+	}
+
+	UDMXPixelMappingOutputComponent* Component = Cast<UDMXPixelMappingOutputComponent>(SelectedComponents.Array()[0].GetComponent());
+	if (!Component)
+	{
+		return;
+	}
+
+	// Apply to parent where appropriate
+	if (Component->GetClass() != UDMXPixelMappingFixtureGroupComponent::StaticClass() &&
+		Component->GetClass() != UDMXPixelMappingMatrixComponent::StaticClass() &&
+		Component->GetClass() != UDMXPixelMappingScreenComponent::StaticClass())
+	{
+		if ((Component->GetParent() && Component->GetParent()->GetClass() == UDMXPixelMappingFixtureGroupComponent::StaticClass()) ||
+			(Component->GetParent() && Component->GetParent()->GetClass() == UDMXPixelMappingMatrixComponent::StaticClass()) ||
+			(Component->GetParent() && Component->GetParent()->GetClass() == UDMXPixelMappingScreenComponent::StaticClass()))
+		{
+			Component = Cast<UDMXPixelMappingOutputComponent>(Component->GetParent());
+		}
+	}
+
+	const FScopedTransaction SizeComponentToTextureTransaction(LOCTEXT("SizeComponentToTextureTransaction", "Size Component to Texture"));
+
+	const UDMXPixelMappingLayoutSettings* LayoutSettings = GetDefault<UDMXPixelMappingLayoutSettings>();
+	if (LayoutSettings && LayoutSettings->bScaleChildrenWithParent)
+	{		
+		// Scale children to parent
+		const FVector2D RatioVector = TextureSize / Component->GetSize();
+		for (UDMXPixelMappingBaseComponent* BaseChild : Component->GetChildren())
+		{
+			if (UDMXPixelMappingOutputComponent* Child = Cast<UDMXPixelMappingOutputComponent>(BaseChild))
+			{
+				Child->PreEditChange(UDMXPixelMappingOutputComponent::StaticClass()->FindPropertyByName(UDMXPixelMappingOutputComponent::GetPositionXPropertyName()));
+				Child->PreEditChange(UDMXPixelMappingOutputComponent::StaticClass()->FindPropertyByName(UDMXPixelMappingOutputComponent::GetPositionYPropertyName()));
+				Child->PreEditChange(UDMXPixelMappingOutputComponent::StaticClass()->FindPropertyByName(UDMXPixelMappingOutputComponent::GetSizeXPropertyName()));
+				Child->PreEditChange(UDMXPixelMappingOutputComponent::StaticClass()->FindPropertyByName(UDMXPixelMappingOutputComponent::GetSizeYPropertyName()));
+
+				// Scale size (SetSize already clamps)
+				Child->SetSize(Child->GetSize() * RatioVector);
+
+				// Scale position (new position is zero vector)
+				const FVector2D ChildPosition = Child->GetPosition();
+				const FVector2D NewPositionRelative = (ChildPosition - Component->GetPosition()) * RatioVector;
+				Child->SetPosition(Component->GetPosition() + NewPositionRelative);
+
+				Child->PostEditChange();
+			}
+		}
+	}
+
+	Component->PreEditChange(UDMXPixelMappingOutputComponent::StaticClass()->FindPropertyByName(UDMXPixelMappingOutputComponent::GetPositionXPropertyName()));
+	Component->PreEditChange(UDMXPixelMappingOutputComponent::StaticClass()->FindPropertyByName(UDMXPixelMappingOutputComponent::GetPositionYPropertyName()));
+	Component->PreEditChange(UDMXPixelMappingOutputComponent::StaticClass()->FindPropertyByName(UDMXPixelMappingOutputComponent::GetSizeXPropertyName()));
+	Component->PreEditChange(UDMXPixelMappingOutputComponent::StaticClass()->FindPropertyByName(UDMXPixelMappingOutputComponent::GetSizeYPropertyName()));
+
+	Component->SetPosition(FVector2D::ZeroVector);
+	Component->SetSize(TextureSize);
+
+	Component->PostEditChange();
 }
 
 void FDMXPixelMappingToolkit::UpdateBlueprintNodes(UDMXPixelMapping* InDMXPixelMapping)
@@ -537,14 +492,10 @@ TArray<UDMXPixelMappingBaseComponent*> FDMXPixelMappingToolkit::CreateComponents
 	{
 		TGuardValue<bool>(bAddingComponents, true);
 
-		FScopedSlowTask SlowTask(Templates.Num(), LOCTEXT("CreateComponentsSlowTaskMessage", "Creating Pixel Mapping Components"));
-		SlowTask.MakeDialogDelayed(.5f);
-
 		if (ensureMsgf(RootComponent && Target, TEXT("Tried to create components from template but RootComponent or Target were invalid.")))
 		{
 			for (const TSharedPtr<FDMXPixelMappingComponentTemplate>& Template : Templates)
 			{
-				SlowTask.EnterProgressFrame(1);
 				if (UDMXPixelMappingBaseComponent* NewComponent = Template->CreateComponent<UDMXPixelMappingBaseComponent>(RootComponent))
 				{
 					NewComponents.Add(NewComponent);
@@ -589,17 +540,6 @@ void FDMXPixelMappingToolkit::DeleteSelectedComponents()
 			return DescendantsCounter;
 		}();
 
-		// Make a slow task since removing may take a significant amount of time
-		const TSharedRef<FScopedSlowTask> SlowTask = MakeShared<FScopedSlowTask>(NumComponentsAndItsDescendants, LOCTEXT("DeleteComponentsSlowTaskMessage", "Removing Pixel Mapping Components"));
-		SlowTask->MakeDialogDelayed(.5f);
-
-		// Bind to the Component Removed Event and let the Slow Task update from that 
-		const FDelegateHandle RemoveComponentDelegateHandle = UDMXPixelMappingBaseComponent::GetOnComponentRemoved().AddRaw(this, &FDMXPixelMappingToolkit::UpdateRemoveComponentSlowTask, SlowTask);
-
-		// Remove the Components
-		const FScopedTransaction DeleteSelectedComponentsTransaction =
-			FScopedTransaction(FText::Format(LOCTEXT("DeleteSelectedComponents", "PixelMapping: Delete {0}|plural(one=Component, other=Components)"), SelectedComponents.Num()));
-
 		TSet<FDMXPixelMappingComponentReference> ParentComponentReferences;
 		for (const FDMXPixelMappingComponentReference& ComponentReference : SelectedComponents)
 		{
@@ -628,9 +568,6 @@ void FDMXPixelMappingToolkit::DeleteSelectedComponents()
 			}
 		}
 
-		// Unbind from the Component Removed Event
-		UDMXPixelMappingBaseComponent::GetOnComponentRemoved().Remove(RemoveComponentDelegateHandle);
-
 		// Select the Parent Component 
 		SelectComponents(ParentComponentReferences);
 
@@ -658,11 +595,6 @@ void FDMXPixelMappingToolkit::OnComponentRemoved(UDMXPixelMapping* PixelMapping,
 	}
 }
 
-void FDMXPixelMappingToolkit::UpdateRemoveComponentSlowTask(UDMXPixelMapping* PixelMapping, UDMXPixelMappingBaseComponent* InComponent, TSharedRef<FScopedSlowTask> SlowTask)
-{
-	SlowTask->EnterProgressFrame(1);
-}
-
 void FDMXPixelMappingToolkit::InitializeInternal(const EToolkitMode::Type Mode, const TSharedPtr<class IToolkitHost>& InitToolkitHost, const FGuid& MessageLogGuid)
 {
 	// Make sure we loaded all UObjects
@@ -727,9 +659,21 @@ void FDMXPixelMappingToolkit::InitializeInternal(const EToolkitMode::Type Mode,
 				)
 				->Split
 				(
-					FTabManager::NewStack()
-					->AddTab(DetailsViewTabID, ETabState::OpenedTab)
-					->SetSizeCoefficient(0.25f)
+					FTabManager::NewSplitter()
+					->SetOrientation(Orient_Vertical)
+					->SetSizeCoefficient(0.5f)
+					->Split
+					(
+						FTabManager::NewStack()
+						->AddTab(DetailsViewTabID, ETabState::OpenedTab)
+						->SetSizeCoefficient(0.6f)
+					)
+					->Split
+					(
+						FTabManager::NewStack()
+						->AddTab(LayoutViewTabID, ETabState::OpenedTab)
+						->SetSizeCoefficient(0.4f)
+					)
 				)
 			)
 		);
@@ -809,6 +753,19 @@ TSharedRef<SDockTab> FDMXPixelMappingToolkit::SpawnTab_DetailsView(const FSpawnT
 	return SpawnedTab;
 }
 
+TSharedRef<SDockTab> FDMXPixelMappingToolkit::SpawnTab_LayoutView(const FSpawnTabArgs& Args)
+{
+	check(Args.GetTabId() == LayoutViewTabID);
+
+	TSharedRef<SDockTab> SpawnedTab = SNew(SDockTab)
+		.Label(LOCTEXT("LayoutViewTabID", "Layout"))
+		[
+			LayoutView.ToSharedRef()
+		];
+
+	return SpawnedTab;
+}
+
 void FDMXPixelMappingToolkit::CreateInternalViewModels()
 {
 	TSharedPtr<FDMXPixelMappingToolkit> ThisPtr(SharedThis(this));
@@ -823,6 +780,7 @@ void FDMXPixelMappingToolkit::CreateInternalViews()
 	CreateOrGetView_DesignerView();
 	CreateOrGetView_PreviewView();
 	CreateOrGetView_DetailsView();
+	CreateOrGetView_LayoutView();
 }
 
 void FDMXPixelMappingToolkit::OnComponentRenamed(UDMXPixelMappingBaseComponent* InComponent)
@@ -880,35 +838,204 @@ TSharedRef<SWidget> FDMXPixelMappingToolkit::CreateOrGetView_DetailsView()
 	return DetailsView.ToSharedRef();
 }
 
+TSharedRef<SWidget> FDMXPixelMappingToolkit::CreateOrGetView_LayoutView()
+{
+	if (!LayoutView.IsValid())
+	{
+		LayoutView = SNew(SDMXPixelMappingLayoutView, SharedThis(this));
+	}
+
+	return LayoutView.ToSharedRef();
+}
+
 void FDMXPixelMappingToolkit::SetupCommands()
 {
 	GetToolkitCommands()->MapAction(
 		FDMXPixelMappingEditorCommands::Get().AddMapping,
-		FExecuteAction::CreateRaw(this, &FDMXPixelMappingToolkit::AddRenderer));
+		FExecuteAction::CreateSP(this, &FDMXPixelMappingToolkit::AddRenderer)
+	);
 
 	GetToolkitCommands()->MapAction(
 		FDMXPixelMappingEditorCommands::Get().SaveThumbnailImage,
-		FExecuteAction::CreateSP(this, &FDMXPixelMappingToolkit::OnSaveThumbnailImage));
+		FExecuteAction::CreateSP(this, &FDMXPixelMappingToolkit::OnSaveThumbnailImage)
+	);
 
 	GetToolkitCommands()->MapAction(
 		FDMXPixelMappingEditorCommands::Get().PlayDMX,
-		FExecuteAction::CreateRaw(this, &FDMXPixelMappingToolkit::PlayDMX),
-		FCanExecuteAction::CreateLambda([this] { return !bIsPlayingDMX; }),
+		FExecuteAction::CreateSP(this, &FDMXPixelMappingToolkit::PlayDMX),
+		FCanExecuteAction::CreateLambda([this] 
+			{ 
+				return !bIsPlayingDMX;
+			}),
 		FIsActionChecked(),
-		FIsActionButtonVisible::CreateLambda([this] { return !bIsPlayingDMX; }));
+		FIsActionButtonVisible::CreateLambda([this] 
+			{ 
+				return !bIsPlayingDMX; 
+			})
+	);
 
 	GetToolkitCommands()->MapAction(
 		FDMXPixelMappingEditorCommands::Get().StopPlayingDMX,
-		FExecuteAction::CreateRaw(this, &FDMXPixelMappingToolkit::StopPlayingDMX),
-		FCanExecuteAction::CreateLambda([this] { return bIsPlayingDMX; }),
+		FExecuteAction::CreateSP(this, &FDMXPixelMappingToolkit::StopPlayingDMX),
+		FCanExecuteAction::CreateLambda([this] 
+			{ 
+				return bIsPlayingDMX; 
+			}),
 		FIsActionChecked(),
-		FIsActionButtonVisible::CreateLambda([this] { return bIsPlayingDMX; }));
+		FIsActionButtonVisible::CreateLambda([this] 
+			{ 
+				return bIsPlayingDMX;
+			})
+	);
+
+	GetToolkitCommands()->MapAction(
+		FDMXPixelMappingEditorCommands::Get().TogglePlayDMXAll,
+		FExecuteAction::CreateSP(this, &FDMXPixelMappingToolkit::ExecutebTogglePlayDMXAll),
+		FCanExecuteAction(),
+		FIsActionChecked::CreateLambda([this]()
+			{ 
+				return bTogglePlayDMXAll; 
+			})
+	);
+
+	// Layout related
+	GetToolkitCommands()->MapAction(
+		FDMXPixelMappingEditorCommands::Get().SizeComponentToTexture,
+		FExecuteAction::CreateSP(this, &FDMXPixelMappingToolkit::SizeSelectedComponentToTexture),
+		FCanExecuteAction::CreateSP(this, &FDMXPixelMappingToolkit::CanSizeSelectedComponentToTexture)
+	);
 
 	GetToolkitCommands()->MapAction(
-		FDMXPixelMappingEditorCommands::Get().bTogglePlayDMXAll,
-		FExecuteAction::CreateRaw(this, &FDMXPixelMappingToolkit::ExecutebTogglePlayDMXAll),
+		FDMXPixelMappingEditorCommands::Get().ToggleScaleChildrenWithParent,
+		FExecuteAction::CreateLambda([]
+			{
+				if (UDMXPixelMappingLayoutSettings* LayoutSettings = GetMutableDefault<UDMXPixelMappingLayoutSettings>())
+				{
+					LayoutSettings->PreEditChange(nullptr);
+					LayoutSettings->bScaleChildrenWithParent = !LayoutSettings->bScaleChildrenWithParent;
+					LayoutSettings->PostEditChange();
+				}
+			}),
 		FCanExecuteAction(),
-		FIsActionChecked::CreateLambda([this]() -> bool { return bTogglePlayDMXAll; }));
+		FIsActionChecked::CreateLambda([]()
+			{
+				if (const UDMXPixelMappingLayoutSettings* LayoutSettings = GetDefault<UDMXPixelMappingLayoutSettings>())
+				{
+					return LayoutSettings->bScaleChildrenWithParent;
+				}
+				return false;
+			})
+	);
+
+	GetToolkitCommands()->MapAction(
+		FDMXPixelMappingEditorCommands::Get().ToggleAlwaysSelectGroup,
+		FExecuteAction::CreateLambda([]
+			{
+				if (UDMXPixelMappingLayoutSettings* LayoutSettings = GetMutableDefault<UDMXPixelMappingLayoutSettings>())
+				{
+					LayoutSettings->PreEditChange(nullptr);
+					LayoutSettings->bAlwaysSelectGroup = !LayoutSettings->bAlwaysSelectGroup;
+					LayoutSettings->PostEditChange();
+				}
+			}),
+		FCanExecuteAction(),
+		FIsActionChecked::CreateLambda([]()
+			{
+				if (const UDMXPixelMappingLayoutSettings* LayoutSettings = GetDefault<UDMXPixelMappingLayoutSettings>())
+				{
+					return LayoutSettings->bAlwaysSelectGroup;
+				}
+				return false;
+			})
+	);
+
+	GetToolkitCommands()->MapAction(
+		FDMXPixelMappingEditorCommands::Get().ToggleApplyLayoutScriptWhenLoaded,
+		FExecuteAction::CreateLambda([]
+			{
+				if (UDMXPixelMappingLayoutSettings* LayoutSettings = GetMutableDefault<UDMXPixelMappingLayoutSettings>())
+				{
+					LayoutSettings->PreEditChange(nullptr);
+					LayoutSettings->bApplyLayoutScriptWhenLoaded = !LayoutSettings->bApplyLayoutScriptWhenLoaded;
+					LayoutSettings->PostEditChange();
+				}
+			}),
+		FCanExecuteAction(),
+		FIsActionChecked::CreateLambda([]()
+			{
+				if (const UDMXPixelMappingLayoutSettings* LayoutSettings = GetDefault<UDMXPixelMappingLayoutSettings>())
+				{
+					return LayoutSettings->bApplyLayoutScriptWhenLoaded;
+				}
+				return false;
+			})
+	);
+
+	GetToolkitCommands()->MapAction(
+		FDMXPixelMappingEditorCommands::Get().ToggleShowComponentNames,
+		FExecuteAction::CreateLambda([]
+			{
+				if (UDMXPixelMappingLayoutSettings* LayoutSettings = GetMutableDefault<UDMXPixelMappingLayoutSettings>())
+				{
+					LayoutSettings->PreEditChange(nullptr);
+					LayoutSettings->bShowComponentNames = !LayoutSettings->bShowComponentNames;
+					LayoutSettings->PostEditChange();
+				}
+			}),
+		FCanExecuteAction(),
+		FIsActionChecked::CreateLambda([]()
+			{
+				if (const UDMXPixelMappingLayoutSettings* LayoutSettings = GetDefault<UDMXPixelMappingLayoutSettings>())
+				{
+					return LayoutSettings->bShowComponentNames;
+				}
+				return false;
+			})
+	);
+
+	GetToolkitCommands()->MapAction(
+		FDMXPixelMappingEditorCommands::Get().ToggleShowPatchInfo,
+		FExecuteAction::CreateLambda([]
+			{
+				if (UDMXPixelMappingLayoutSettings* LayoutSettings = GetMutableDefault<UDMXPixelMappingLayoutSettings>())
+				{
+					LayoutSettings->PreEditChange(nullptr);
+					LayoutSettings->bShowPatchInfo = !LayoutSettings->bShowPatchInfo;
+					LayoutSettings->PostEditChange();
+				}
+			}),
+		FCanExecuteAction(),
+		FIsActionChecked::CreateLambda([]()
+			{
+				if (const UDMXPixelMappingLayoutSettings* LayoutSettings = GetDefault<UDMXPixelMappingLayoutSettings>())
+				{
+					return LayoutSettings->bShowPatchInfo;
+				}
+				return false;
+			})
+	);
+
+	GetToolkitCommands()->MapAction(
+		FDMXPixelMappingEditorCommands::Get().ToggleShowCellIDs,
+		FExecuteAction::CreateLambda([]
+			{
+				if (UDMXPixelMappingLayoutSettings* LayoutSettings = GetMutableDefault<UDMXPixelMappingLayoutSettings>())
+				{
+					LayoutSettings->PreEditChange(nullptr);
+					LayoutSettings->bShowCellIDs = !LayoutSettings->bShowCellIDs;
+					LayoutSettings->PostEditChange();
+				}
+			}),
+		FCanExecuteAction(),
+		FIsActionChecked::CreateLambda([]()
+			{
+				if (const UDMXPixelMappingLayoutSettings* LayoutSettings = GetDefault<UDMXPixelMappingLayoutSettings>())
+				{
+					return LayoutSettings->bShowCellIDs;
+				}
+				return false;
+			})
+	);
 }
 
 void FDMXPixelMappingToolkit::ExtendToolbar()
diff --git a/Engine/Plugins/VirtualProduction/DMX/DMXPixelMapping/Source/DMXPixelMappingEditor/Private/Toolkits/DMXPixelMappingToolkit.h b/Engine/Plugins/VirtualProduction/DMX/DMXPixelMapping/Source/DMXPixelMappingEditor/Private/Toolkits/DMXPixelMappingToolkit.h
--- a/Engine/Plugins/VirtualProduction/DMX/DMXPixelMapping/Source/DMXPixelMappingEditor/Private/Toolkits/DMXPixelMappingToolkit.h
+++ b/Engine/Plugins/VirtualProduction/DMX/DMXPixelMapping/Source/DMXPixelMappingEditor/Private/Toolkits/DMXPixelMappingToolkit.h
@@ -8,13 +8,14 @@
 #include "Toolkits/AssetEditorToolkit.h"
 #include "DMXPixelMappingComponentReference.h"
 
-class FDMXPixelMappingPaletteViewModel;
 class FDMXPixelMappingToolbar;
-class SDMXPixelMappingPaletteView;
 class SDMXPixelMappingHierarchyView;
 class SDMXPixelMappingDesignerView;
-class SDMXPixelMappingPreviewView;
 class SDMXPixelMappingDetailsView;
+class SDMXPixelMappingLayoutView;
+class SDMXPixelMappingPaletteView;
+class SDMXPixelMappingPreviewView;
+class FDMXPixelMappingPaletteViewModel;
 class UDMXPixelMapping;
 class UDMXPixelMappingBaseComponent;
 class UDMXPixelMappingMatrixComponent;
@@ -37,10 +38,6 @@ class FDMXPixelMappingToolkit
 	friend class FDMXPixelMappingToolbar;
 
 public:
-	DECLARE_MULTICAST_DELEGATE(FOnComponentsAddedOrDeletedDelegate)
-	UE_DEPRECATED(4.27, "Use UDMXPixelMappingBaseComponent::GetOnComponentAdded and  UDMXPixelMappingBaseComponent::GetOnComponentRemoved instead.")
-	FOnComponentsAddedOrDeletedDelegate& GetOnComponentsAddedOrDeletedDelegate() { return OnComponentsAddedOrDeletedDelegate_DEPRECATED; }
-
 	DECLARE_MULTICAST_DELEGATE(FOnSelectedComponentsChangedDelegate)
 	FOnSelectedComponentsChangedDelegate& GetOnSelectedComponentsChangedDelegate() { return OnSelectedComponentsChangedDelegate; }
 
@@ -109,28 +106,28 @@ public:
 
 	const TSet<FDMXPixelMappingComponentReference>& GetSelectedComponents() const { return SelectedComponents; }
 
+	/** Gets or creates the Palette View for this Pixel Mapping instance */
 	TSharedRef<SWidget> CreateOrGetView_PaletteView();
 
+	/** Gets or creates the Hierarchy View for this Pixel Mapping instance */
 	TSharedRef<SWidget> CreateOrGetView_HierarchyView();
 
+	/** Gets or creates the Designer View for this Pixel Mapping instance */
 	TSharedRef<SWidget> CreateOrGetView_DesignerView();
 
+	/** Gets or creates the Preview View for this Pixel Mapping instance */
 	TSharedRef<SWidget> CreateOrGetView_PreviewView();
 
+	/** Gets or creates the Details View for this Pixel Mapping instance */
 	TSharedRef<SWidget> CreateOrGetView_DetailsView();
 
+	/** Gets or creates the Layout View for this Pixel Mapping instance */
+	TSharedRef<SWidget> CreateOrGetView_LayoutView();
+
 	bool IsPlayingDMX() const { return bIsPlayingDMX; }
 
 	void SetActiveRenderComponent(UDMXPixelMappingRendererComponent* InComponent);
 
-	/** DEPRECATED 4.27 */
-	UE_DEPRECATED(4.27, "No longer needed")
-	void HandleAddComponents();
-
-	/** DEPRECATED 4.27 */
-	UE_DEPRECATED(4.27, "No longer needed")
-	void HandleRemoveComponents();
-
 	/** Creates an array of components given specifed component references */
 	template <typename ComponentType>
 	TArray<ComponentType> MakeComponentArray(const TSet<FDMXPixelMappingComponentReference>& Components) const;
@@ -140,9 +137,6 @@ public:
 	/** Returns true if the component is selected */
 	bool IsComponentSelected(UDMXPixelMappingBaseComponent* Component) const;
 
-	/** Gives each component widget a color depending on selection and whether it's over its parent */
-	void UpdateComponentWidgetColors();
-
 	void AddRenderer();
 
 	void OnComponentRenamed(UDMXPixelMappingBaseComponent* InComponent);
@@ -163,31 +157,42 @@ private:
 	/** Called when a Component was removed from the pixel mapping */
 	void OnComponentRemoved(UDMXPixelMapping* PixelMapping, UDMXPixelMappingBaseComponent* Component);
 
-	/** Updates the Slow Task when removing Components */
-	void UpdateRemoveComponentSlowTask(UDMXPixelMapping* PixelMapping, UDMXPixelMappingBaseComponent* Component, TSharedRef<FScopedSlowTask> SlowTask);
-
 	void PlayDMX();
 
 	void StopPlayingDMX();
 
 	void ExecutebTogglePlayDMXAll();
 
+	/** Returns true if the selected component can be sized to texture */
+	bool CanSizeSelectedComponentToTexture() const;
+
+	/** Sizes the component to the render target of the pixelmapping asset */
+	void SizeSelectedComponentToTexture();
+
 	void UpdateBlueprintNodes(UDMXPixelMapping* InDMXPixelMapping);
 
 	void OnSaveThumbnailImage();
 
 	void InitializeInternal(const EToolkitMode::Type Mode, const TSharedPtr<class IToolkitHost>& InitToolkitHost, const FGuid& MessageLogGuid);
 
+	/** Spawns the Pallette View */
 	TSharedRef<SDockTab> SpawnTab_PaletteView(const FSpawnTabArgs& Args);
 
+	/** Spawns the Hierarchy View */
 	TSharedRef<SDockTab> SpawnTab_HierarchyView(const FSpawnTabArgs& Args);
 
+	/** Spawns the Designer View */
 	TSharedRef<SDockTab> SpawnTab_DesignerView(const FSpawnTabArgs& Args);
 
+	/** Spawns the Preview View */
 	TSharedRef<SDockTab> SpawnTab_PreviewView(const FSpawnTabArgs& Args);
 
+	/** Spawns the Details View */
 	TSharedRef<SDockTab> SpawnTab_DetailsView(const FSpawnTabArgs& Args);
 
+	/** Spawns the Layout View */
+	TSharedRef<SDockTab> SpawnTab_LayoutView(const FSpawnTabArgs& Args);
+
 	void SetupCommands();
 
 	void ExtendToolbar();
@@ -201,16 +206,25 @@ private:
 	/** List of open tool panels; used to ensure only one exists at any one time */
 	TMap<FName, TWeakPtr<SDockableTab>> SpawnedToolPanels;
 
+	/** The Palette View instance */
 	TSharedPtr<SDMXPixelMappingPaletteView> PaletteView;
-
+	
+	/** The Hierarchy View instance */
 	TSharedPtr<SDMXPixelMappingHierarchyView> HierarchyView;
 
+	/** The Designer View instance */
 	TSharedPtr<SDMXPixelMappingDesignerView> DesignerView;
 
+	/** The Preview View instance */
 	TSharedPtr<SDMXPixelMappingPreviewView> PreviewView;
 
+	/** The Details View instance */
 	TSharedPtr<SDMXPixelMappingDetailsView> DetailsView;
 
+	/** The Layout View instance */
+	TSharedPtr<SDMXPixelMappingLayoutView> LayoutView;
+
+	/** The Palette View Model instance */
 	TSharedPtr<FDMXPixelMappingPaletteViewModel> PaletteViewModel;
 
 	FOnSelectedComponentsChangedDelegate OnSelectedComponentsChangedDelegate;
@@ -242,16 +256,22 @@ private:
 	/** True while removing components (to avoid needlessly updating blueprint nodes on each component removed via our own methods) */
 	bool bRemovingComponents = false;
 
-	FOnComponentsAddedOrDeletedDelegate OnComponentsAddedOrDeletedDelegate_DEPRECATED;
-
 public:
+	/** Name of the Palette View Tab */
 	static const FName PaletteViewTabID;
 
+	/** Name of the Hierarchy View Tab */
 	static const FName HierarchyViewTabID;
 
+	/** Name of the Designer View Tab */
 	static const FName DesignerViewTabID;
 
+	/** Name of the Preview View Tab */
 	static const FName PreviewViewTabID;
 
+	/** Name of the Details View Tab */
 	static const FName DetailsViewTabID;
+
+	/** Name of the Details View Tab */
+	static const FName LayoutViewTabID;
 };
