DMX - Add layout options for DMX Pixelmapping, component widgets now use a model view pattern

diff --git a/Engine/Plugins/VirtualProduction/DMX/DMXPixelMapping/Source/DMXPixelMappingEditor/Private/ViewModels/DMXPixelMappingOutputComponentModel.cpp b/Engine/Plugins/VirtualProduction/DMX/DMXPixelMapping/Source/DMXPixelMappingEditor/Private/ViewModels/DMXPixelMappingOutputComponentModel.cpp
new file mode 100644
--- /dev/null
+++ b/Engine/Plugins/VirtualProduction/DMX/DMXPixelMapping/Source/DMXPixelMappingEditor/Private/ViewModels/DMXPixelMappingOutputComponentModel.cpp
@@ -0,0 +1,381 @@
+// Copyright Epic Games, Inc. All Rights Reserved.
+
+#include "ViewModels/DMXPixelMappingOutputComponentModel.h"
+
+#include "DMXPixelMappingComponentReference.h"
+#include "DMXPixelMappingTypes.h"
+#include "Components/DMXPixelMappingOutputComponent.h"
+#include "Components/DMXPixelMappingFixtureGroupItemComponent.h"
+#include "Components/DMXPixelMappingMatrixCellComponent.h"
+#include "Components/DMXPixelMappingMatrixComponent.h"
+#include "Components/DMXPixelMappingScreenComponent.h"
+#include "Library/DMXEntityFixturePatch.h"
+#include "Library/DMXLibrary.h"
+#include "MVR/Types/DMXMVRFixtureNode.h"
+#include "Toolkits/DMXPixelMappingToolkit.h"
+
+
+#define LOCTEXT_NAMESPACE "DMXPixelMappingOutputComponentModel"
+
+FDMXPixelMappingOutputComponentModel::FDMXPixelMappingOutputComponentModel(const TSharedRef<FDMXPixelMappingToolkit>& InToolkit, TWeakObjectPtr<UDMXPixelMappingOutputComponent> InOutputComponent)
+	: WeakOutputComponent(InOutputComponent)
+	, WeakToolkit(InToolkit)
+{
+	UDMXPixelMappingOutputComponent* OutputComponent = WeakOutputComponent.Get();
+	if (!OutputComponent)
+	{
+		return;
+	}
+
+	UpdateFixtureNode();
+	bSelected = InToolkit->GetSelectedComponents().Contains(FDMXPixelMappingComponentReference(InToolkit, OutputComponent));
+
+	// Handle selection changes
+	InToolkit->GetOnSelectedComponentsChangedDelegate().AddRaw(this, &FDMXPixelMappingOutputComponentModel::OnSelectedComponentsChanged);
+}
+
+FDMXPixelMappingOutputComponentModel::~FDMXPixelMappingOutputComponentModel()
+{
+	if (TSharedPtr<FDMXPixelMappingToolkit> Toolkit = WeakToolkit.Pin())
+	{
+		Toolkit->GetOnSelectedComponentsChangedDelegate().RemoveAll(this);
+	}
+}
+
+FVector2D FDMXPixelMappingOutputComponentModel::GetPosition() const
+{
+	if (UDMXPixelMappingOutputComponent* OutputComponent = WeakOutputComponent.Get())
+	{
+		return OutputComponent->GetPosition();
+	}
+	return FVector2D::ZeroVector;
+}
+
+FVector2D FDMXPixelMappingOutputComponentModel::GetSize() const
+{
+	if (UDMXPixelMappingOutputComponent* OutputComponent = WeakOutputComponent.Get())
+	{
+		return OutputComponent->GetSize();
+	}
+	return FVector2D::ZeroVector;
+}
+
+FText FDMXPixelMappingOutputComponentModel::GetName() const
+{
+	if (UDMXPixelMappingOutputComponent* OutputComponent = Cast<UDMXPixelMappingOutputComponent>(WeakOutputComponent.Get()))
+	{
+		return FText::FromString(OutputComponent->GetName());
+	}
+
+	return FText::GetEmpty();
+}
+
+bool FDMXPixelMappingOutputComponentModel::ShouldDrawName() const
+{
+	if (UDMXPixelMappingOutputComponent* OutputComponent = Cast<UDMXPixelMappingOutputComponent>(WeakOutputComponent.Get()))
+	{
+		return OutputComponent->GetClass() != UDMXPixelMappingMatrixCellComponent::StaticClass();
+	}
+	return true;
+}
+
+bool FDMXPixelMappingOutputComponentModel::ShouldDrawNameAbove() const
+{
+	if (UDMXPixelMappingOutputComponent* OutputComponent = Cast<UDMXPixelMappingOutputComponent>(WeakOutputComponent.Get()))
+	{
+		return
+			OutputComponent->GetClass() == UDMXPixelMappingFixtureGroupComponent::StaticClass() ||
+			OutputComponent->GetClass() == UDMXPixelMappingScreenComponent::StaticClass();
+	}
+	return true;
+}
+
+bool FDMXPixelMappingOutputComponentModel::HasPatchInfo() const
+{
+	return GetFixturePatch() != nullptr;
+}
+
+FText FDMXPixelMappingOutputComponentModel::GetAddressesText() const
+{
+	if (const UDMXEntityFixturePatch* FixturePatch = GetFixturePatch())
+	{
+		const FString AddressesString = FString::Printf(TEXT("%i.%i"), FixturePatch->GetUniverseID(), FixturePatch->GetStartingChannel());
+		
+		return FText::Format(LOCTEXT("PatchTextWithLabel", "Patch: {0}"), FText::FromString(AddressesString));
+	}
+
+	return FText::GetEmpty();
+}
+
+
+FText FDMXPixelMappingOutputComponentModel::GetFixtureIDText() const
+{
+	if (UDMXMVRFixtureNode* FixtureNode = WeakFixtureNode.Get())
+	{
+		const FText FixtureIDText = FixtureNode->FixtureID.IsEmpty() ? LOCTEXT("NoFixtureID", "<not set>") : FText::FromString(FixtureNode->FixtureID);
+		return FText::Format(LOCTEXT("FixtureIDTextWithLabel", "FID: {0}"), FixtureIDText);
+	}
+
+	return FText::GetEmpty();
+}
+
+bool FDMXPixelMappingOutputComponentModel::HasCellID() const
+{
+	if (UDMXPixelMappingMatrixCellComponent* MatrixCell = Cast<UDMXPixelMappingMatrixCellComponent>(WeakOutputComponent.Get()))
+	{
+		return true;
+	}
+	return false;
+}
+
+FText FDMXPixelMappingOutputComponentModel::GetCellIDText() const
+{
+	if (UDMXPixelMappingMatrixCellComponent* MatrixCell = Cast<UDMXPixelMappingMatrixCellComponent>(WeakOutputComponent.Get()))
+	{
+		return FText::FromString(FString::FromInt(MatrixCell->CellID));
+	}
+
+	return FText::GetEmpty();
+}
+
+FLinearColor FDMXPixelMappingOutputComponentModel::GetColor() const
+{
+	if (UDMXPixelMappingOutputComponent* OutputComponent = WeakOutputComponent.Get())
+	{
+		if (OutputComponent->IsOverParent())
+		{
+			if (bSelected)
+			{
+				return OutputComponent->GetEditorColor();
+			}
+			else
+			{
+				return OutputComponent->GetEditorColor().CopyWithNewOpacity(0.6f);
+			}
+		}
+	}
+
+	return FLinearColor::Red;
+}
+
+bool FDMXPixelMappingOutputComponentModel::Equals(UDMXPixelMappingBaseComponent* Other) const
+{
+	return WeakOutputComponent.Get() == Other; 
+}
+
+void FDMXPixelMappingOutputComponentModel::OnSelectedComponentsChanged()
+{
+	TSharedPtr<FDMXPixelMappingToolkit> Toolkit = WeakToolkit.Pin();
+	if (!Toolkit.IsValid())
+	{
+		return;
+	}
+
+	if (UDMXPixelMappingOutputComponent* OutputComponent = WeakOutputComponent.Get())
+	{
+		bSelected = Toolkit->GetSelectedComponents().Contains(FDMXPixelMappingComponentReference(Toolkit, OutputComponent));
+	}
+}
+
+const UDMXEntityFixturePatch* FDMXPixelMappingOutputComponentModel::GetFixturePatch() const
+{
+	const UDMXEntityFixturePatch* FixturePatch = nullptr;
+	if (UDMXPixelMappingMatrixComponent* Matrix = Cast<UDMXPixelMappingMatrixComponent>(WeakOutputComponent.Get()))
+	{
+		FixturePatch = Matrix->FixturePatchRef.GetFixturePatch();
+	}
+	else if (UDMXPixelMappingFixtureGroupItemComponent* GroupItem = Cast<UDMXPixelMappingFixtureGroupItemComponent>(WeakOutputComponent.Get()))
+	{
+		FixturePatch = GroupItem->FixturePatchRef.GetFixturePatch();
+	}
+
+	return FixturePatch;
+}
+
+void FDMXPixelMappingOutputComponentModel::UpdateFixtureNode()
+{
+	WeakFixtureNode.Reset();
+
+	const UDMXEntityFixturePatch* FixturePatch = GetFixturePatch();
+	if (!FixturePatch)
+	{
+		return;
+	}
+
+	const FGuid& MVRUUID = FixturePatch->GetMVRFixtureUUID();
+	if (!MVRUUID.IsValid())
+	{
+		return;
+	}
+
+	const UDMXLibrary* DMXLibrary = FixturePatch->GetParentLibrary();
+	if (!DMXLibrary)
+	{
+		return;
+	}
+
+	// We can use the lazy General Scene Description
+	UDMXMVRGeneralSceneDescription* GeneralSceneDescription = DMXLibrary->GetLazyGeneralSceneDescription();
+	if (!GeneralSceneDescription)
+	{
+		return;
+	}
+
+	GeneralSceneDescription->WriteDMXLibraryToGeneralSceneDescription(*DMXLibrary);
+	WeakFixtureNode = GeneralSceneDescription->FindFixtureNode(MVRUUID);
+}
+
+
+FDMXPixelMappingScreenComponentModel::FDMXPixelMappingScreenComponentModel(const TSharedRef<FDMXPixelMappingToolkit>& InToolkit, TWeakObjectPtr<UDMXPixelMappingScreenComponent> InScreenComponent)
+	: WeakScreenComponent(InScreenComponent)
+	, WeakToolkit(InToolkit)
+{
+	UDMXPixelMappingScreenComponent* ScreenComponent = WeakScreenComponent.Get();
+	if (!ScreenComponent)
+	{
+		return;
+	}
+
+	bSelected = InToolkit->GetSelectedComponents().Contains(FDMXPixelMappingComponentReference(InToolkit, ScreenComponent));
+
+	// Handle selection changes
+	InToolkit->GetOnSelectedComponentsChangedDelegate().AddRaw(this, &FDMXPixelMappingScreenComponentModel::OnSelectedComponentsChanged);
+}
+
+FDMXPixelMappingScreenComponentModel::~FDMXPixelMappingScreenComponentModel()
+{
+	if (TSharedPtr<FDMXPixelMappingToolkit> Toolkit = WeakToolkit.Pin())
+	{
+		Toolkit->GetOnSelectedComponentsChangedDelegate().RemoveAll(this);
+	}
+}
+
+FVector2D FDMXPixelMappingScreenComponentModel::GetPosition() const
+{
+	if (UDMXPixelMappingScreenComponent* ScreenComponent = WeakScreenComponent.Get())
+	{
+		return ScreenComponent->GetPosition();
+	}
+	return FVector2D::ZeroVector;
+}
+
+FVector2D FDMXPixelMappingScreenComponentModel::GetSize() const
+{
+	if (UDMXPixelMappingScreenComponent* ScreenComponent = WeakScreenComponent.Get())
+	{
+		return ScreenComponent->GetSize();
+	}
+	return FVector2D::ZeroVector;
+}
+
+int32 FDMXPixelMappingScreenComponentModel::GetNumColumns() const
+{
+	if (UDMXPixelMappingScreenComponent* ScreenComponent = WeakScreenComponent.Get())
+	{
+		return ScreenComponent->NumXCells;
+	}
+	return 1;
+}
+
+int32 FDMXPixelMappingScreenComponentModel::GetNumRows() const
+{
+	if (UDMXPixelMappingScreenComponent* ScreenComponent = WeakScreenComponent.Get())
+	{
+		return ScreenComponent->NumYCells;
+	}
+	return 1;
+}
+
+FLinearColor FDMXPixelMappingScreenComponentModel::GetColor() const
+{
+	if (UDMXPixelMappingScreenComponent* ScreenComponent = WeakScreenComponent.Get())
+	{
+		if (bSelected)
+		{
+			return ScreenComponent->GetEditorColor();
+		}
+		else
+		{
+			return ScreenComponent->GetEditorColor().CopyWithNewOpacity(0.6f);
+		}
+	}
+
+	return FLinearColor::Red;
+}
+
+bool FDMXPixelMappingScreenComponentModel::Equals(UDMXPixelMappingBaseComponent* Other) const
+{
+	return WeakScreenComponent.Get() == Other;
+}
+
+void FDMXPixelMappingScreenComponentModel::OnSelectedComponentsChanged()
+{
+	TSharedPtr<FDMXPixelMappingToolkit> Toolkit = WeakToolkit.Pin();
+	if (!Toolkit.IsValid())
+	{
+		return;
+	}
+
+	if (UDMXPixelMappingScreenComponent* ScreenComponent = WeakScreenComponent.Get())
+	{
+		bSelected = Toolkit->GetSelectedComponents().Contains(FDMXPixelMappingComponentReference(Toolkit, ScreenComponent));
+	}
+}
+
+EDMXPixelMappingDistribution FDMXPixelMappingScreenComponentModel::GetDistribution() const
+{
+	if (UDMXPixelMappingScreenComponent* ScreenComponent = WeakScreenComponent.Get())
+	{
+		return ScreenComponent->Distribution;
+	}
+	return EDMXPixelMappingDistribution::BottomLeftToRight;
+}
+
+EDMXCellFormat FDMXPixelMappingScreenComponentModel::GetCellFormat() const
+{
+	if (UDMXPixelMappingScreenComponent* ScreenComponent = WeakScreenComponent.Get())
+	{
+		return ScreenComponent->PixelFormat;
+	}
+	return EDMXCellFormat::PF_RGB;
+}
+
+bool FDMXPixelMappingScreenComponentModel::ComponentWantsToShowUniverse() const
+{
+	if (UDMXPixelMappingScreenComponent* ScreenComponent = WeakScreenComponent.Get())
+	{
+		return ScreenComponent->bShowUniverse;
+	}
+	return false;
+}
+
+int32 FDMXPixelMappingScreenComponentModel::GetUniverse() const
+{
+	if (UDMXPixelMappingScreenComponent* ScreenComponent = WeakScreenComponent.Get())
+	{
+		return ScreenComponent->LocalUniverse;
+	}
+
+	return 0;
+}
+
+bool FDMXPixelMappingScreenComponentModel::ComponentWantsToShowChannel() const
+{
+	if (UDMXPixelMappingScreenComponent* ScreenComponent = WeakScreenComponent.Get())
+	{
+		return ScreenComponent->bShowAddresses;
+	}
+	return false;
+}
+
+int32 FDMXPixelMappingScreenComponentModel::GetStartingChannel() const
+{
+	if (UDMXPixelMappingScreenComponent* ScreenComponent = WeakScreenComponent.Get())
+	{
+		return ScreenComponent->StartAddress;
+	}
+
+	return 0;
+}
+
+#undef LOCTEXT_NAMESPACE
diff --git a/Engine/Plugins/VirtualProduction/DMX/DMXPixelMapping/Source/DMXPixelMappingEditor/Private/Viewports/DMXPixelMappingSceneViewport.h b/Engine/Plugins/VirtualProduction/DMX/DMXPixelMapping/Source/DMXPixelMappingEditor/Private/Viewports/DMXPixelMappingSceneViewport.h
--- a/Engine/Plugins/VirtualProduction/DMX/DMXPixelMapping/Source/DMXPixelMappingEditor/Private/Viewports/DMXPixelMappingSceneViewport.h
+++ b/Engine/Plugins/VirtualProduction/DMX/DMXPixelMapping/Source/DMXPixelMappingEditor/Private/Viewports/DMXPixelMappingSceneViewport.h
@@ -40,9 +40,9 @@ protected:
 
 	virtual FReply OnMouseWheel(const FGeometry& InGeometry, const FPointerEvent& MouseEvent) override
 	{
-		FSceneViewport::OnMouseButtonDown(InGeometry, MouseEvent);
+		FSceneViewport::OnMouseWheel(InGeometry, MouseEvent);
 
-		return FReply::Unhandled();
+		return FReply::Handled();
 	}
 
 	virtual FReply OnMouseButtonDoubleClick(const FGeometry& InGeometry, const FPointerEvent& InMouseEvent) override
diff --git a/Engine/Plugins/VirtualProduction/DMX/DMXPixelMapping/Source/DMXPixelMappingEditor/Private/Viewports/DMXPixelMappingSourceTextureViewportClient.cpp b/Engine/Plugins/VirtualProduction/DMX/DMXPixelMapping/Source/DMXPixelMappingEditor/Private/Viewports/DMXPixelMappingSourceTextureViewportClient.cpp
--- a/Engine/Plugins/VirtualProduction/DMX/DMXPixelMapping/Source/DMXPixelMappingEditor/Private/Viewports/DMXPixelMappingSourceTextureViewportClient.cpp
+++ b/Engine/Plugins/VirtualProduction/DMX/DMXPixelMapping/Source/DMXPixelMappingEditor/Private/Viewports/DMXPixelMappingSourceTextureViewportClient.cpp
@@ -33,10 +33,10 @@ void FDMXPixelMappingSourceTextureViewportClient::Draw(FViewport* InViewport, FC
 		InputTexture = RendererComponent->GetRendererInputTexture();
 	}
 
-	if (InputTexture != nullptr && InputTexture->GetResource() != nullptr)
+	if (InputTexture && InputTexture->GetResource())
 	{
-		uint32 Width = SceneViewport->GetSizeXY().X;
-		uint32 Height = SceneViewport->GetSizeXY().Y;
+		const uint32 Width = SceneViewport->GetSizeXY().X;
+		const uint32 Height = SceneViewport->GetSizeXY().Y;
 
 		FCanvasTileItem TileItem(FVector2D(0, 0), InputTexture->GetResource(), FVector2D(Width, Height), FLinearColor::White);
 		TileItem.BlendMode = ESimpleElementBlendMode::SE_BLEND_MAX;
diff --git a/Engine/Plugins/VirtualProduction/DMX/DMXPixelMapping/Source/DMXPixelMappingEditor/Private/Views/SDMXPixelMappingDesignerView.cpp b/Engine/Plugins/VirtualProduction/DMX/DMXPixelMapping/Source/DMXPixelMappingEditor/Private/Views/SDMXPixelMappingDesignerView.cpp
--- a/Engine/Plugins/VirtualProduction/DMX/DMXPixelMapping/Source/DMXPixelMappingEditor/Private/Views/SDMXPixelMappingDesignerView.cpp
+++ b/Engine/Plugins/VirtualProduction/DMX/DMXPixelMapping/Source/DMXPixelMappingEditor/Private/Views/SDMXPixelMappingDesignerView.cpp
@@ -6,6 +6,7 @@
 #include "DMXPixelMappingComponentWidget.h"
 #include "DMXPixelMappingEditorUtils.h"
 #include "DMXPixelMappingEditorStyle.h"
+#include "DMXPixelMappingLayoutSettings.h"
 #include "SDMXPixelMappingComponentBox.h"
 #include "Components/DMXPixelMappingOutputComponent.h"
 #include "Components/DMXPixelMappingRendererComponent.h"
@@ -14,10 +15,12 @@
 #include "Components/DMXPixelMappingFixtureGroupItemComponent.h"
 #include "Components/DMXPixelMappingMatrixComponent.h"
 #include "Components/DMXPixelMappingMatrixCellComponent.h"
+#include "Components/DMXPixelMappingScreenComponent.h"
 #include "DragDrop/DMXPixelMappingDragDropOp.h"
 #include "Library/DMXEntityFixturePatch.h"
 #include "Library/DMXLibrary.h"
 #include "Widgets/SDMXPixelMappingDesignerCanvas.h"
+#include "Widgets/SDMXPixelMappingOutputComponent.h"
 #include "Widgets/SDMXPixelMappingRuler.h"
 #include "Widgets/SDMXPixelMappingSourceTextureViewport.h"
 #include "Widgets/SDMXPixelMappingTransformHandle.h"
@@ -276,60 +279,55 @@ FReply SDMXPixelMappingDesignerView::OnMouseButtonDown(const FGeometry& MyGeomet
 {
 	SDMXPixelMappingSurface::OnMouseButtonDown(MyGeometry, MouseEvent);
 
-	if (TSharedPtr<FDMXPixelMappingToolkit> ToolkitPtr = ToolkitWeakPtr.Pin())
+	const bool bLeftMouseButton = MouseEvent.GetEffectingButton() == EKeys::LeftMouseButton;
+	UDMXPixelMappingOutputComponent* ClickedComponent = GetComponentUnderCursor();
+	if (bLeftMouseButton && ClickedComponent)
 	{
-		if (MouseEvent.GetEffectingButton() == EKeys::LeftMouseButton)
+		PendingSelectedComponent = ClickedComponent;
+
+		// Handle reclicking a selected component depending on type
+		const bool bClickedSelectedComponent = [ClickedComponent, this]()
 		{
-			// Select and detect drag when something was clicked
-			if (UDMXPixelMappingOutputComponent* ClickedComponent = GetComponentUnderCursor())
+			for (const FDMXPixelMappingComponentReference& SelectedComponentRef : GetSelectedComponents())
 			{
-				PendingSelectedComponent = ClickedComponent;
-
-				// Handle reclicking a selected component depending on type
-				const bool bClickedSelectedComponent = [ClickedComponent, this]()
+				if (ClickedComponent == SelectedComponentRef.GetComponent())
+				{
+					return true;
+				}
+				else if (UDMXPixelMappingMatrixCellComponent* ClickedMatrixCellComponent = Cast<UDMXPixelMappingMatrixCellComponent>(ClickedComponent))
 				{
-					for (const FDMXPixelMappingComponentReference& SelectedComponentRef : GetSelectedComponents())
+					if (ClickedMatrixCellComponent->GetParent() == SelectedComponentRef.GetComponent())
 					{
-						if (ClickedComponent == SelectedComponentRef.GetComponent())
-						{
-							return true;
-						}
-						else if (UDMXPixelMappingMatrixCellComponent* ClickedMatrixCellComponent = Cast<UDMXPixelMappingMatrixCellComponent>(ClickedComponent))
-						{
-							if (ClickedMatrixCellComponent->GetParent() == SelectedComponentRef.GetComponent())
-							{
-								return true;
-							}
-						}
-						else if (UDMXPixelMappingFixtureGroupItemComponent* ClickedGroupItemComponent = Cast<UDMXPixelMappingFixtureGroupItemComponent>(ClickedComponent))
-						{
-							if (ClickedGroupItemComponent->GetParent() == SelectedComponentRef.GetComponent() &&
-								ClickedGroupItemComponent->IsLockInDesigner())
-							{
-								return true;
-							}
-						}
+						return true;
 					}
-
-					return false;
-				}();
-
-				const bool bClearPreviousSelection = !MouseEvent.IsShiftDown() && !MouseEvent.IsControlDown() && !bClickedSelectedComponent;
-				ResolvePendingSelectedComponents(bClearPreviousSelection);
-
-				FVector2D GraphSpaceCursorPosition;
-				if (GetGraphSpaceCursorPosition(GraphSpaceCursorPosition))
+				}
+				else if (UDMXPixelMappingFixtureGroupItemComponent* ClickedGroupItemComponent = Cast<UDMXPixelMappingFixtureGroupItemComponent>(ClickedComponent))
 				{
-					DragAnchor = GraphSpaceCursorPosition;
-
-					return
-						FReply::Handled()
-						.PreventThrottling()
-						.SetUserFocus(AsShared(), EFocusCause::Mouse)
-						.CaptureMouse(AsShared())
-						.DetectDrag(AsShared(), EKeys::LeftMouseButton);
+					if (ClickedGroupItemComponent->GetParent() == SelectedComponentRef.GetComponent() &&
+						ClickedGroupItemComponent->IsLockInDesigner())
+					{
+						return true;
+					}
 				}
 			}
+
+			return false;
+		}();
+
+		const bool bClearPreviousSelection = !MouseEvent.IsShiftDown() && !MouseEvent.IsControlDown() && !bClickedSelectedComponent;
+		ResolvePendingSelectedComponents(bClearPreviousSelection);
+
+		FVector2D GraphSpaceCursorPosition;
+		if (GetGraphSpaceCursorPosition(GraphSpaceCursorPosition))
+		{
+			DragAnchor = GraphSpaceCursorPosition;
+
+			return
+				FReply::Handled()
+				.PreventThrottling()
+				.SetUserFocus(AsShared(), EFocusCause::Mouse)
+				.CaptureMouse(AsShared())
+				.DetectDrag(AsShared(), EKeys::LeftMouseButton);
 		}
 	}
 
@@ -363,7 +361,10 @@ FReply SDMXPixelMappingDesignerView::OnMouseButtonUp(const FGeometry& MyGeometry
 	const bool bClearPreviousSelection = !MouseEvent.IsShiftDown() && !MouseEvent.IsControlDown();
 	ResolvePendingSelectedComponents(bClearPreviousSelection);
 
-	return FReply::Handled().ReleaseMouseCapture();
+	return FReply::Handled()
+		.EndDragDrop()
+		.ReleaseMouseCapture()
+		.SetUserFocus(AsShared());
 }
 
 FReply SDMXPixelMappingDesignerView::OnMouseMove(const FGeometry& MyGeometry, const FPointerEvent& MouseEvent)
@@ -675,6 +676,12 @@ void SDMXPixelMappingDesignerView::PostRedo(bool bSuccess)
 
 void SDMXPixelMappingDesignerView::RebuildDesigner()
 {
+	const TSharedPtr<FDMXPixelMappingToolkit> Toolkit = ToolkitWeakPtr.Pin();
+	if (!Toolkit.IsValid())
+	{
+		return;
+	}
+
 	if (DesignCanvas.IsValid())
 	{
 		DesignCanvas->ClearChildren();
@@ -684,20 +691,31 @@ void SDMXPixelMappingDesignerView::RebuildDesigner()
 	(
 		SAssignNew(DesignCanvas, SDMXPixelMappingDesignerCanvas)
 	);
-	const TSharedRef<SConstraintCanvas> ConstraintCanvas = SNew(SConstraintCanvas);
-	DesignCanvas->AddSlot()
-		[
-			ConstraintCanvas
-		];
 
+	OutputComponentWidgets.Reset();
 	if (CachedRendererComponent.IsValid())
 	{
-		CachedRendererComponent->ForEachChild([ConstraintCanvas](UDMXPixelMappingBaseComponent* InComponent)
+		const TSharedRef<SConstraintCanvas> ComponentCanvas = SNew(SConstraintCanvas);
+		DesignCanvas->AddSlot()
+			[
+				ComponentCanvas
+			];
+
+		CachedRendererComponent->ForEachChild([this, &Toolkit, &ComponentCanvas](UDMXPixelMappingBaseComponent* Component)
 			{
-				if (UDMXPixelMappingOutputComponent* Component = Cast<UDMXPixelMappingOutputComponent>(InComponent))
+				if (UDMXPixelMappingScreenComponent* ScreenComponent = Cast<UDMXPixelMappingScreenComponent>(Component))
+				{
+					const TSharedRef<SDMXPixelMappingScreenComponent> ScreenComponentWidget = SNew(SDMXPixelMappingScreenComponent, Toolkit.ToSharedRef(), ScreenComponent);
+					ScreenComponentWidget->AddToCanvas(ComponentCanvas);
+
+					OutputComponentWidgets.Add(ScreenComponentWidget);
+				}
+				else if (UDMXPixelMappingOutputComponent* OutputComponent = Cast<UDMXPixelMappingOutputComponent>(Component))
 				{
-					// Build all child DMX pixel mapping slots
-					Component->BuildSlot(ConstraintCanvas);
+					const TSharedRef<SDMXPixelMappingOutputComponent> ComponentWidget = SNew(SDMXPixelMappingOutputComponent, Toolkit.ToSharedRef(), OutputComponent);
+					ComponentWidget->AddToCanvas(ComponentCanvas);
+
+					OutputComponentWidgets.Add(ComponentWidget);
 				}
 			}, true);
 	}
@@ -894,37 +912,41 @@ void SDMXPixelMappingDesignerView::OnSelectedComponentsChanged()
 
 void SDMXPixelMappingDesignerView::ResolvePendingSelectedComponents(bool bClearPreviousSelection)
 {
-	if (TSharedPtr<FDMXPixelMappingToolkit> ToolkitPtr = ToolkitWeakPtr.Pin())
+	const TSharedPtr<FDMXPixelMappingToolkit> Toolkit = ToolkitWeakPtr.Pin();
+	const bool bHasPendingSelectedComponent = PendingSelectedComponent.IsValid();
+	if (Toolkit.IsValid() && bHasPendingSelectedComponent)
 	{
-		if (PendingSelectedComponent.IsValid())
+		// Never select matrix cells
+		if (UDMXPixelMappingMatrixCellComponent* MatrixCellComponent = Cast<UDMXPixelMappingMatrixCellComponent>(PendingSelectedComponent))
 		{
-			if (UDMXPixelMappingMatrixCellComponent* MatrixCellComponent = Cast<UDMXPixelMappingMatrixCellComponent>(PendingSelectedComponent))
+			PendingSelectedComponent = MatrixCellComponent->GetParent();
+		}
+
+		// Select group if desired
+		const UDMXPixelMappingLayoutSettings* LayoutSettings = GetDefault<UDMXPixelMappingLayoutSettings>();
+		if (LayoutSettings && LayoutSettings->bAlwaysSelectGroup)
+		{
+			if (UDMXPixelMappingFixtureGroupItemComponent* GroupItemComponent = Cast<UDMXPixelMappingFixtureGroupItemComponent>(PendingSelectedComponent))
 			{
-				// If a Matrix Cell is selected, select the owning Matrix Component instead
-				PendingSelectedComponent = MatrixCellComponent->GetParent();
+				PendingSelectedComponent = GroupItemComponent->GetParent();
 			}
-			else if (UDMXPixelMappingFixtureGroupItemComponent* GroupItemComponent = Cast<UDMXPixelMappingFixtureGroupItemComponent>(PendingSelectedComponent))
+			else if (UDMXPixelMappingMatrixComponent* MatrixComponent = Cast<UDMXPixelMappingMatrixComponent>(PendingSelectedComponent))
 			{
-				// If an Item Component is selected and is lock in designer, select the owning Group Component instead
-				if (GroupItemComponent->IsLockInDesigner())
-				{
-					PendingSelectedComponent = GroupItemComponent->GetParent();
-				}
+				PendingSelectedComponent = MatrixComponent->GetParent();
 			}
+		}
 
-			TSet<FDMXPixelMappingComponentReference> SelectedComponents;
-
-			// Add the newly selected component first. This is important, e.g. for drag drop when iterating and using this as base
-			SelectedComponents.Add(ToolkitWeakPtr.Pin()->GetReferenceFromComponent(PendingSelectedComponent.Get()));
+		// Add the newly selected component first. This is important, e.g. for drag drop when iterating and using this as base
+		TSet<FDMXPixelMappingComponentReference> SelectedComponents;
+		SelectedComponents.Add(Toolkit->GetReferenceFromComponent(PendingSelectedComponent.Get()));
 
-			if (!bClearPreviousSelection)
-			{
-				SelectedComponents.Append(ToolkitWeakPtr.Pin()->GetSelectedComponents());
-			}
-			ToolkitWeakPtr.Pin()->SelectComponents(SelectedComponents);
-
-			PendingSelectedComponent = nullptr;
+		if (!bClearPreviousSelection)
+		{
+			SelectedComponents.Append(ToolkitWeakPtr.Pin()->GetSelectedComponents());
 		}
+		Toolkit->SelectComponents(SelectedComponents);
+
+		PendingSelectedComponent = nullptr;
 	}
 }
 
@@ -967,6 +989,8 @@ UDMXPixelMappingOutputComponent* SDMXPixelMappingDesignerView::GetComponentUnder
 		FVector2D GraphSpaceCursorPosition;
 		if (GetGraphSpaceCursorPosition(GraphSpaceCursorPosition))
 		{
+			// Widgets are offset by 0.5f to reside over the pixel
+			GraphSpaceCursorPosition += FVector2D(.5f, .5f);
 			RendererComponent->ForEachChild([&ComponentUnderCursor, &GraphSpaceCursorPosition](UDMXPixelMappingBaseComponent* InComponent)
 				{
 					if (UDMXPixelMappingOutputComponent* OutputComponent = Cast<UDMXPixelMappingOutputComponent>(InComponent))
@@ -1028,17 +1052,18 @@ bool SDMXPixelMappingDesignerView::GetComponentGeometry(UDMXPixelMappingBaseComp
 {
 	if (UDMXPixelMappingOutputComponent* OutputComponent = Cast<UDMXPixelMappingOutputComponent>(InBaseComponent))
 	{
-		if (TSharedPtr<FDMXPixelMappingComponentWidget> ComponentWidget = OutputComponent->GetComponentWidget())
+		const TSharedRef<IDMXPixelMappingOutputComponentWidgetInterface>* ComponentViewPtr = Algo::FindByPredicate(OutputComponentWidgets, [InBaseComponent](const TSharedRef<IDMXPixelMappingOutputComponentWidgetInterface>& ComponentWidget)
+			{
+				return ComponentWidget->Equals(InBaseComponent);
+			});
+
+		if (ComponentViewPtr)
 		{
-			TSharedPtr<SWidget> CachedPreviewWidget = ComponentWidget->GetComponentBox();
-			if (CachedPreviewWidget.IsValid())
+			const FArrangedWidget* ArrangedWidget = CachedWidgetGeometry.Find((*ComponentViewPtr)->AsWidget());
+			if (ArrangedWidget)
 			{
-				const FArrangedWidget* ArrangedWidget = CachedWidgetGeometry.Find(CachedPreviewWidget.ToSharedRef());
-				if (ArrangedWidget)
-				{
-					OutGeometry = ArrangedWidget->Geometry;
-					return true;
-				}
+				OutGeometry = ArrangedWidget->Geometry;
+				return true;
 			}
 		}
 	}
diff --git a/Engine/Plugins/VirtualProduction/DMX/DMXPixelMapping/Source/DMXPixelMappingEditor/Private/Views/SDMXPixelMappingDesignerView.h b/Engine/Plugins/VirtualProduction/DMX/DMXPixelMapping/Source/DMXPixelMappingEditor/Private/Views/SDMXPixelMappingDesignerView.h
--- a/Engine/Plugins/VirtualProduction/DMX/DMXPixelMapping/Source/DMXPixelMappingEditor/Private/Views/SDMXPixelMappingDesignerView.h
+++ b/Engine/Plugins/VirtualProduction/DMX/DMXPixelMapping/Source/DMXPixelMappingEditor/Private/Views/SDMXPixelMappingDesignerView.h
@@ -9,17 +9,16 @@
 #include "EditorUndoClient.h"
 #include "Templates/SharedPointer.h"
 
-class UDMXPixelMappingBaseComponent;
 class FDMXPixelMappingComponentReference;
 class FDMXPixelMappingDragDropOp;
-class UDMXPixelMappingFixtureGroupComponent;
-class UDMXPixelMappingFixtureGroupItemComponent;
-class UDMXPixelMappingOutputComponent;
-class UDMXPixelMappingRendererComponent;
+class FDMXPixelMappingToolkit;
+class IDMXPixelMappingOutputComponentWidgetInterface;
 class SDMXPixelMappingRuler;
 class SDMXPixelMappingSourceTextureViewport;
-class FDMXPixelMappingToolkit;
 class SDMXPixelMappingTransformHandle;
+class UDMXPixelMappingBaseComponent;
+class UDMXPixelMappingOutputComponent;
+class UDMXPixelMappingRendererComponent;
 
 class FHittestGrid;
 struct FOptionalSize;
@@ -220,6 +219,9 @@ private:
 	/** The position in local space where the user began dragging a widget */
 	FVector2D DragAnchor;
 
+	/** Oupput component views currently displayed */
+	TArray<TSharedRef<IDMXPixelMappingOutputComponentWidgetInterface>> OutputComponentWidgets;
+
 	TArray<TSharedPtr<SDMXPixelMappingTransformHandle>> TransformHandles;
 
 	/** Helper class to restore selection on scope */
diff --git a/Engine/Plugins/VirtualProduction/DMX/DMXPixelMapping/Source/DMXPixelMappingEditor/Private/Views/SDMXPixelMappingLayoutView..cpp b/Engine/Plugins/VirtualProduction/DMX/DMXPixelMapping/Source/DMXPixelMappingEditor/Private/Views/SDMXPixelMappingLayoutView..cpp
new file mode 100644
--- /dev/null
+++ b/Engine/Plugins/VirtualProduction/DMX/DMXPixelMapping/Source/DMXPixelMappingEditor/Private/Views/SDMXPixelMappingLayoutView..cpp
@@ -0,0 +1,166 @@
+// Copyright Epic Games, Inc. All Rights Reserved.
+
+#include "Views/SDMXPixelMappingLayoutView.h"
+
+#include "DMXPixelMappingLayoutSettings.h"
+#include "Customizations/DMXPixelMappingHorizontalAlignmentCustomization.h"
+#include "Customizations/DMXPixelMappingLayoutViewModelDetails.h"
+#include "Customizations/DMXPixelMappingVerticalAlignmentCustomization.h"
+#include "Toolkits/DMXPixelMappingToolkit.h"
+#include "ViewModels/DMXPixelMappingLayoutViewModel.h"
+
+#include "DetailsViewArgs.h"
+#include "Editor.h"
+#include "IDetailsView.h"
+#include "PropertyEditorModule.h"
+#include "TimerManager.h"
+#include "Modules/ModuleManager.h"
+#include "Widgets/Input/SCheckBox.h"
+#include "Widgets/Layout/SBorder.h"
+#include "Widgets/Layout/SScrollBox.h"
+
+
+#define LOCTEXT_NAMESPACE "SDMXPixelMappingHierarchyView"
+
+SDMXPixelMappingLayoutView::~SDMXPixelMappingLayoutView()
+{
+	if (GEditor)
+	{
+		GEditor->UnregisterForUndo(this);
+	}
+}
+
+void SDMXPixelMappingLayoutView::Construct(const FArguments& InArgs, const TSharedRef<FDMXPixelMappingToolkit>& InToolkit)
+{
+	WeakToolkit = InToolkit;
+
+	Model = NewObject<UDMXPixelMappingLayoutViewModel>();
+	Model->SetToolkit(InToolkit);
+	Model->OnModelChanged.AddSP(this, &SDMXPixelMappingLayoutView::Refresh);
+
+	InitializeLayoutScriptsDetailsView();
+
+	if (!ensureMsgf(LayoutScriptDetailsView.IsValid(), TEXT("Failed to create Details View for Pixel Mapping Layouts")))
+	{
+		return;
+	}
+
+	ChildSlot
+		[
+			SNew(SScrollBox)
+
+			// Layout Model Details (select Layout Script class etc.)
+			+ SScrollBox::Slot()
+			[
+				CreateLayoutModelDetailsView()
+			]
+
+			// Layout Script Details
+			+ SScrollBox::Slot()
+			[
+				LayoutScriptDetailsView.ToSharedRef()
+			]
+
+			// Layout Settings Details
+			+ SScrollBox::Slot()
+			[
+				CreateLayoutSettingsDetailsView()
+			]
+
+		];
+
+	Refresh();
+
+	if (GEditor)
+	{
+		GEditor->RegisterForUndo(this);
+	}
+}
+
+void SDMXPixelMappingLayoutView::AddReferencedObjects(FReferenceCollector& Collector)
+{
+	Collector.AddReferencedObject(Model);
+}
+
+void SDMXPixelMappingLayoutView::PostUndo(bool bSuccess)
+{
+	GEditor->GetTimerManager()->SetTimerForNextTick(FTimerDelegate::CreateSP(this, &SDMXPixelMappingLayoutView::Refresh));
+}
+
+void SDMXPixelMappingLayoutView::PostRedo(bool bSuccess)
+{
+	GEditor->GetTimerManager()->SetTimerForNextTick(FTimerDelegate::CreateSP(this, &SDMXPixelMappingLayoutView::Refresh));
+}
+
+void SDMXPixelMappingLayoutView::NotifyPostChange(const FPropertyChangedEvent& PropertyChangedEvent, FProperty* PropertyThatChanged)
+{
+	Model->ApplyLayoutScripts();
+}
+
+void SDMXPixelMappingLayoutView::Refresh()
+{
+	const TArray<UObject*> LayoutScriptObjects = Model->GetLayoutScriptsObjectsSlow();
+	if (LayoutScriptObjects.IsEmpty())
+	{
+		LayoutScriptDetailsView->SetVisibility(EVisibility::Collapsed);
+	}
+	else
+	{
+		LayoutScriptDetailsView->SetObjects(LayoutScriptObjects);
+		LayoutScriptDetailsView->ForceRefresh();
+
+		LayoutScriptDetailsView->SetVisibility(EVisibility::Visible);
+	}
+}
+
+TSharedRef<SWidget> SDMXPixelMappingLayoutView::CreateLayoutSettingsDetailsView()
+{
+	UDMXPixelMappingLayoutSettings* LayoutSettings = GetMutableDefault<UDMXPixelMappingLayoutSettings>();
+	if (!LayoutSettings)
+	{
+		return SNullWidget::NullWidget;
+	}
+
+	FPropertyEditorModule& PropertyEditorModule = FModuleManager::GetModuleChecked<FPropertyEditorModule>("PropertyEditor");
+	FDetailsViewArgs DetailsViewArgs;
+	DetailsViewArgs.bAllowSearch = false;
+	DetailsViewArgs.NameAreaSettings = FDetailsViewArgs::HideNameArea;
+	TSharedRef<IDetailsView> LayoutSettingsDetailsView = PropertyEditorModule.CreateDetailView(DetailsViewArgs);
+
+	LayoutSettingsDetailsView->SetObject(LayoutSettings);
+
+	return LayoutSettingsDetailsView;
+}
+
+TSharedRef<SWidget> SDMXPixelMappingLayoutView::CreateLayoutModelDetailsView()
+{
+	FPropertyEditorModule& PropertyEditorModule = FModuleManager::GetModuleChecked<FPropertyEditorModule>("PropertyEditor");
+	FDetailsViewArgs DetailsViewArgs;
+	DetailsViewArgs.bAllowSearch = false;
+	DetailsViewArgs.NameAreaSettings = FDetailsViewArgs::HideNameArea;
+	TSharedRef<IDetailsView> ModelDetailsView = PropertyEditorModule.CreateDetailView(DetailsViewArgs);
+	
+	ModelDetailsView->RegisterInstancedCustomPropertyLayout(UDMXPixelMappingLayoutViewModel::StaticClass(), FOnGetDetailCustomizationInstance::CreateStatic(&FDMXPixelMappingLayoutViewModelDetails::MakeInstance, WeakToolkit));
+
+	ModelDetailsView->SetObject(Model);
+
+	return ModelDetailsView;
+}
+
+void SDMXPixelMappingLayoutView::InitializeLayoutScriptsDetailsView()
+{
+	TArray<UObject*> LayoutScripts = Model->GetLayoutScriptsObjectsSlow();
+
+	FPropertyEditorModule& PropertyEditorModule = FModuleManager::GetModuleChecked<FPropertyEditorModule>("PropertyEditor");
+	FDetailsViewArgs DetailsViewArgs;
+	DetailsViewArgs.bAllowSearch = false;
+	DetailsViewArgs.NotifyHook = this;
+
+	DetailsViewArgs.NameAreaSettings = FDetailsViewArgs::HideNameArea;
+	LayoutScriptDetailsView = PropertyEditorModule.CreateDetailView(DetailsViewArgs);
+
+	LayoutScriptDetailsView->RegisterInstancedCustomPropertyTypeLayout(TEXT("EHorizontalAlignment"), FOnGetPropertyTypeCustomizationInstance::CreateStatic(&FDMXPixelMappingHorizontalAlignmentCustomization::MakeInstance));
+	LayoutScriptDetailsView->RegisterInstancedCustomPropertyTypeLayout(TEXT("EVerticalAlignment"), FOnGetPropertyTypeCustomizationInstance::CreateStatic(&FDMXPixelMappingVerticalAlignmentCustomization::MakeInstance));
+}
+
+#undef LOCTEXT_NAMESPACE
diff --git a/Engine/Plugins/VirtualProduction/DMX/DMXPixelMapping/Source/DMXPixelMappingEditor/Private/Widgets/SDMXPixelMappingFixturePatchDetailRow.cpp b/Engine/Plugins/VirtualProduction/DMX/DMXPixelMapping/Source/DMXPixelMappingEditor/Private/Widgets/SDMXPixelMappingFixturePatchDetailRow.cpp
--- a/Engine/Plugins/VirtualProduction/DMX/DMXPixelMapping/Source/DMXPixelMappingEditor/Private/Widgets/SDMXPixelMappingFixturePatchDetailRow.cpp
+++ b/Engine/Plugins/VirtualProduction/DMX/DMXPixelMapping/Source/DMXPixelMappingEditor/Private/Widgets/SDMXPixelMappingFixturePatchDetailRow.cpp
@@ -43,30 +43,36 @@ void SDMXPixelMappingFixturePatchDetailRow::Construct(const FArguments& InArgs)
 					return bHighlight ? HighlightBGColor : NormalBGColor;
 				})
 			[
-				SAssignNew(FixturePatchNameTextBlock, STextBlock)
-				.Text_Lambda([FixturePatch]()
-					{
-						const UDMXEntityFixtureType* FixtureType = FixturePatch.IsValid() ? FixturePatch->GetFixtureType() : nullptr;
-						const FDMXFixtureMode* FixtureModePtr = FixturePatch.IsValid() ? FixturePatch->GetActiveMode() : nullptr;
-						if (FixtureType && FixtureModePtr)
+				SNew(SBorder)
+				.BorderImage(FAppStyle::GetBrush("NoBorder"))
+				.VAlign(VAlign_Center)
+				.HAlign(HAlign_Left)
+				[
+					SAssignNew(FixturePatchNameTextBlock, STextBlock)
+					.Text_Lambda([FixturePatch]()
 						{
-							if (FixtureModePtr->bFixtureMatrixEnabled)
+							const UDMXEntityFixtureType* FixtureType = FixturePatch.IsValid() ? FixturePatch->GetFixtureType() : nullptr;
+							const FDMXFixtureMode* FixtureModePtr = FixturePatch.IsValid() ? FixturePatch->GetActiveMode() : nullptr;
+							if (FixtureType && FixtureModePtr)
 							{
-								return FText::Format(LOCTEXT("MatrixFixturePatchName", "Matrix: {0}"), FText::FromString(FixturePatch->Name));
+								if (FixtureModePtr->bFixtureMatrixEnabled)
+								{
+									return FText::Format(LOCTEXT("MatrixFixturePatchName", "Matrix: {0}"), FText::FromString(FixturePatch->Name));
+								}
+								else
+								{
+									return FText::FromString(FixturePatch->Name);
+								}
 							}
-							else
-							{
-								return FText::FromString(FixturePatch->Name);
-							}
-						}
-						// The user object should take care of not showing invalid patches.
-						return FText::GetEmpty();
-					})
-				.Font(FAppStyle::GetFontStyle("PropertyWindow.NormalFont"))
-				.ColorAndOpacity_Lambda([this]()
-					{
-						return bHighlight ? FLinearColor::Black : FLinearColor::White;
-					})
+							// The user object should take care of not showing invalid patches.
+							return FText::GetEmpty();
+						})
+					.Font(FAppStyle::GetFontStyle("PropertyWindow.NormalFont"))
+					.ColorAndOpacity_Lambda([this]()
+						{
+							return bHighlight ? FLinearColor::Black : FLinearColor::White;
+						})
+				]
 			]
 		]
 	];
diff --git a/Engine/Plugins/VirtualProduction/DMX/DMXPixelMapping/Source/DMXPixelMappingEditor/Private/Widgets/SDMXPixelMappingOutputComponent.cpp b/Engine/Plugins/VirtualProduction/DMX/DMXPixelMapping/Source/DMXPixelMappingEditor/Private/Widgets/SDMXPixelMappingOutputComponent.cpp
new file mode 100644
--- /dev/null
+++ b/Engine/Plugins/VirtualProduction/DMX/DMXPixelMapping/Source/DMXPixelMappingEditor/Private/Widgets/SDMXPixelMappingOutputComponent.cpp
@@ -0,0 +1,703 @@
+// Copyright Epic Games, Inc. All Rights Reserved.
+
+#include "Widgets/SDMXPixelMappingOutputComponent.h"
+
+#include "DMXPixelMappingLayoutSettings.h"
+#include "DMXPixelMappingTypes.h"
+#include "DMXPixelMappingUtils.h"
+#include "DMXRuntimeUtils.h"
+#include "Components/DMXPixelMappingOutputComponent.h"
+#include "Components/DMXPixelMappingScreenComponent.h"
+#include "ViewModels/DMXPixelMappingOutputComponentModel.h"
+#include "Toolkits/DMXPixelMappingToolkit.h"
+
+#include "IDetailsView.h"
+#include "PropertyEditorModule.h"
+#include "TimerManager.h"
+#include "Modules/ModuleManager.h"
+#include "Widgets/Layout/SBox.h"
+#include "Widgets/Layout/SScaleBox.h"
+#include "Widgets/Layout/SUniformGridPanel.h"
+
+
+#define LOCTEXT_NAMESPACE "SDMXPixelMappingOutputComponent"
+
+IDMXPixelMappingOutputComponentWidgetInterface::~IDMXPixelMappingOutputComponentWidgetInterface()
+{
+	RemoveFromCanvas();
+}
+
+void IDMXPixelMappingOutputComponentWidgetInterface::AddToCanvas(const TSharedRef<SConstraintCanvas>& InCanvas)
+{
+	RemoveFromCanvas();
+
+	ParentCanvas = InCanvas;
+
+	ParentCanvas->AddSlot()
+		.ZOrder(100)
+		.AutoSize(true)
+		.Alignment(FVector2D::ZeroVector)
+		.Offset_Lambda([this]()
+			{
+				const FVector2D Postition = GetPosition();
+
+				// In the center of the top left pixel
+				return FMargin(Postition.X + .5f, Postition.Y + .5f);
+			})
+		.Expose(Slot)
+		[
+			AsWidget()
+		];
+}
+
+void IDMXPixelMappingOutputComponentWidgetInterface::RemoveFromCanvas()
+{
+	if (ParentCanvas.IsValid() && Slot)
+	{
+		ParentCanvas->RemoveSlot(Slot->GetWidget());
+	}
+
+	Slot = nullptr;
+	ParentCanvas.Reset();
+}
+
+void SDMXPixelMappingOutputComponent::Construct(const FArguments& InArgs, const TSharedRef<FDMXPixelMappingToolkit>& InToolkit, TWeakObjectPtr<UDMXPixelMappingOutputComponent> OutputComponent)
+{
+	if (!OutputComponent.IsValid())
+	{
+		return;
+	}
+
+	WeakToolkit = InToolkit;
+	Model = MakeShared<FDMXPixelMappingOutputComponentModel>(InToolkit, OutputComponent);
+
+	BorderBrush.DrawAs = ESlateBrushDrawType::Border;
+	BorderBrush.Margin = FMargin(1.f);
+
+	ChildSlot
+	[
+		SAssignNew(ComponentBox, SBox)
+		.WidthOverride_Lambda([this]()
+			{
+				return Model->GetSize().X;
+			})
+		.HeightOverride_Lambda([this]()
+			{
+				return Model->GetSize().Y;
+			})
+		[
+			SNew(SBorder)
+			.BorderImage_Lambda([this]()
+				{
+					BorderBrush.TintColor = Model->GetColor();
+					return &BorderBrush;
+				})
+			[
+				SNew(SVerticalBox)
+
+				+ SVerticalBox::Slot()
+				.HAlign(HAlign_Fill)
+				.VAlign(VAlign_Bottom)
+				.AutoHeight()
+				.Padding(FMargin(0.f, -22.f, 0.f, 0.f))
+				[
+					SNew(SBox)
+					.HeightOverride(22.f)
+					.HAlign(HAlign_Left)
+					.VAlign(VAlign_Bottom)
+					[
+						SAssignNew(AboveContentBorder, SBorder)
+						.BorderImage(FAppStyle::GetBrush("NoBorder"))
+					]
+				]
+
+				+ SVerticalBox::Slot()
+				.HAlign(HAlign_Fill)
+				.VAlign(VAlign_Fill)
+				.FillHeight(1.f / 3.f)
+				.Padding(FMargin(-2.f, -2.f, -2.f, -2.f))
+				[
+					SAssignNew(TopContentBorder, SBorder)
+					.BorderImage(FAppStyle::GetBrush("NoBorder"))
+				]
+
+				+ SVerticalBox::Slot()
+				.HAlign(HAlign_Fill)
+				.VAlign(VAlign_Fill)
+				.FillHeight(1.f / 3.f)
+				.Padding(FMargin(-2.f, -2.f, -2.f, -2.f))
+				[
+					SAssignNew(MiddleContentBorder, SBorder)
+					.BorderImage(FAppStyle::GetBrush("NoBorder"))
+				]
+
+				+ SVerticalBox::Slot()
+				.HAlign(HAlign_Fill)
+				.VAlign(VAlign_Fill)
+				.FillHeight(1.f / 3.f)
+				.Padding(FMargin(-2.f, -2.f, -2.f, -2.f))
+				[
+					SAssignNew(BottomContentBorder, SBorder)
+					.BorderImage(FAppStyle::GetBrush("NoBorder"))
+				]
+			]
+		]
+	];
+
+	UpdateChildSlots();
+	UDMXPixelMappingLayoutSettings::GetOnLayoutSettingsChanged().AddSP(this, &SDMXPixelMappingOutputComponent::UpdateChildSlots);
+}
+
+bool SDMXPixelMappingOutputComponent::Equals(UDMXPixelMappingBaseComponent* Component) const
+{
+	if (Component)
+	{
+		return Model->Equals(Component);
+	}
+	return false;
+}
+
+FVector2D SDMXPixelMappingOutputComponent::GetPosition() const
+{
+	return Model->GetPosition();
+}
+
+void SDMXPixelMappingOutputComponent::UpdateChildSlots()
+{
+	AboveContentBorder->ClearContent();
+	TopContentBorder->ClearContent();
+	MiddleContentBorder->ClearContent();
+	BottomContentBorder->ClearContent();
+
+	const UDMXPixelMappingLayoutSettings* LayoutSettings = GetDefault<UDMXPixelMappingLayoutSettings>();
+	if (!LayoutSettings)
+	{
+		return;
+	}
+
+	if (LayoutSettings->bShowComponentNames && Model->ShouldDrawName())
+	{
+		if (Model->ShouldDrawNameAbove())
+		{
+			CreateComponentNameChildSlotAbove();
+		}
+		else
+		{
+			CreateComponentNameChildSlotInside();
+		}
+	}
+
+	if (LayoutSettings->bShowCellIDs && Model->HasCellID())
+	{
+		CreateCellIDChildSlot();
+	}
+
+	if (LayoutSettings->bShowPatchInfo && Model->HasPatchInfo())
+	{
+		CreatePatchInfoChildSlot();
+	}
+}
+
+void SDMXPixelMappingOutputComponent::CreateComponentNameChildSlotAbove()
+{
+	AboveContentBorder->SetContent(
+		SNew(SBox)
+		.HeightOverride(20.f)
+		.HAlign(HAlign_Left)
+		.VAlign(VAlign_Bottom)
+		[
+			SNew(SScaleBox)
+			.Stretch(EStretch::ScaleToFit)
+			.StretchDirection(EStretchDirection::DownOnly)
+			[
+				SNew(STextBlock)
+				.Font(FCoreStyle::GetDefaultFontStyle("Regular", 18))
+				.Text_Lambda([this]()
+					{
+						return Model->GetName();
+					})
+			]	
+		]
+
+	);
+}
+
+void SDMXPixelMappingOutputComponent::CreateComponentNameChildSlotInside()
+{
+	TopContentBorder->SetContent(
+		SNew(SBox)
+		.HeightOverride(14.f)
+		.HAlign(HAlign_Left)
+		.VAlign(VAlign_Top)
+		[
+			SNew(SScaleBox)
+			.Stretch(EStretch::ScaleToFit)
+			.StretchDirection(EStretchDirection::DownOnly)
+			.HAlign(HAlign_Left)
+			.VAlign(VAlign_Top)
+			[
+				SNew(STextBlock)
+				.Font(FCoreStyle::GetDefaultFontStyle("Regular", 12))
+				.Text_Lambda([this]()
+					{
+						return Model->GetName();
+					})
+			]
+		]
+	);
+}
+
+void SDMXPixelMappingOutputComponent::CreateCellIDChildSlot()
+{	
+	MiddleContentBorder->SetContent(
+		SNew(SBox)
+		.HeightOverride(14.f)
+		.HAlign(HAlign_Center)
+		.VAlign(VAlign_Center)
+		[
+			SNew(SScaleBox)
+			.Stretch(EStretch::ScaleToFit)
+			.StretchDirection(EStretchDirection::DownOnly)
+			[
+				SNew(STextBlock)
+				.Font(FCoreStyle::GetDefaultFontStyle("Regular", 12))
+				.Text_Lambda([this]()
+					{
+						return Model->GetCellIDText();
+					})
+			]
+		]
+	);
+}
+
+void SDMXPixelMappingOutputComponent::CreatePatchInfoChildSlot()
+{
+	BottomContentBorder->SetContent(
+		SNew(SBox)
+		.HeightOverride(8.f)
+		.HAlign(HAlign_Left)
+		.VAlign(VAlign_Bottom)
+		[
+			SNew(SScaleBox)
+			.Stretch(EStretch::ScaleToFit)
+			.StretchDirection(EStretchDirection::DownOnly)
+			[
+				SNew(SVerticalBox)
+
+				+ SVerticalBox::Slot()
+				.HAlign(HAlign_Left)
+				.VAlign(VAlign_Bottom)
+				[
+					SNew(STextBlock)
+					.ColorAndOpacity(FLinearColor::Green)
+					.Font(FCoreStyle::GetDefaultFontStyle("Regular", 6))
+					.Text_Lambda([this]()
+						{
+							return Model->GetAddressesText();
+						})
+				]
+
+				+ SVerticalBox::Slot()
+				.HAlign(HAlign_Left)
+				.VAlign(VAlign_Bottom)
+				[
+					SNew(STextBlock)
+					.ColorAndOpacity(FLinearColor::Green)
+					.Font(FCoreStyle::GetDefaultFontStyle("Regular", 6))
+					.Text_Lambda([this]()
+						{
+							return Model->GetFixtureIDText();
+						})
+				]
+			]
+		]
+	);
+}
+
+
+class SDMXPixelMappingScreenComponentCell
+	: public SCompoundWidget
+{
+public:
+	SLATE_BEGIN_ARGS(SDMXPixelMappingScreenComponentCell)
+		: _ChannelOffset(0)
+		, _StartingChannel(1)
+		, _Universe(1)
+	{}
+		/** The offset of this cell from the starting address */
+		SLATE_ARGUMENT(int32, ChannelOffset)
+
+		/** The starting address of the screen component */
+		SLATE_ATTRIBUTE(int32, StartingChannel)
+
+		/** The universe of the screen component */
+		SLATE_ATTRIBUTE(int32, Universe)
+
+	SLATE_END_ARGS()
+
+	/** Constructs this widget */
+	void Construct(const FArguments& InArgs, const TSharedRef<FDMXPixelMappingScreenComponentModel>& InModel)
+	{
+		Model = InModel;
+
+		ChannelOffset = InArgs._ChannelOffset;
+		Universe = InArgs._Universe;
+		StartingChannel = InArgs._StartingChannel;
+
+		BorderBrush.DrawAs = ESlateBrushDrawType::Border;
+		BorderBrush.Margin = FMargin(1.f);
+		BorderBrush.TintColor = FLinearColor::White.CopyWithNewOpacity(.4f);
+
+		ChildSlot
+			[
+				SAssignNew(ContentBorder, SBorder)
+				.BorderImage(&BorderBrush)
+				.HAlign(HAlign_Fill)
+				.VAlign(VAlign_Fill)
+			];
+
+		RebuildAddressesContent();
+
+		UDMXPixelMappingLayoutSettings::GetOnLayoutSettingsChanged().AddSP(this, &SDMXPixelMappingScreenComponentCell::RebuildAddressesContent);
+	}
+
+private:
+	/** Rebuilds addresses content */
+	void RebuildAddressesContent()
+	{
+		const UDMXPixelMappingLayoutSettings* LayoutSettings = GetDefault<UDMXPixelMappingLayoutSettings>();
+		if (!LayoutSettings || !LayoutSettings->bShowPatchInfo)
+		{
+			return;
+		}
+
+		if (!ensureMsgf(StartingChannel.IsSet() && Universe.IsSet(), TEXT("Trying to display universe and address of pixel mapping screen component, but universe and/or address are not set.")))
+		{
+			return;
+		}
+
+		ContentBorder->SetContent
+		(
+			SNew(SBorder)
+			.BorderImage(FAppStyle::GetBrush("NoBorder"))
+			.HAlign(HAlign_Fill)
+			.VAlign(VAlign_Fill)
+			.Padding(6.f)
+			[
+				SNew(SScaleBox)
+				.Stretch(EStretch::ScaleToFit)
+				.StretchDirection(EStretchDirection::DownOnly)
+				.VAlign(VAlign_Center)
+				.HAlign(HAlign_Center)
+				[
+					SNew(STextBlock)
+					.Font(FCoreStyle::GetDefaultFontStyle("Regular", 12))
+					.Text_Lambda([this]()
+						{
+							const int32 UniverseWithOffset = Universe.Get() + ChannelOffset / DMX_MAX_ADDRESS;
+							const int32 ChannelWithOffset = [this, UniverseWithOffset]()
+							{
+								if (UniverseWithOffset > 0)
+								{
+									return (StartingChannel.Get() + ChannelOffset) % DMX_MAX_ADDRESS;
+								}
+								return StartingChannel.Get() + ChannelOffset;
+							}();
+
+							if (Model->ComponentWantsToShowChannel() && Model->ComponentWantsToShowUniverse())
+							{
+								FString AddressString = FString::FromInt(UniverseWithOffset) + TEXT(".") + FString::FromInt(ChannelWithOffset);
+								return FText::FromString(AddressString);
+							}
+							else if (Model->ComponentWantsToShowUniverse())
+							{
+								FString UniverseString = FString::FromInt(UniverseWithOffset);
+								return FText::FromString(UniverseString);
+							}
+							else if (Model->ComponentWantsToShowChannel())
+							{
+								FString ChannelString = FString::FromInt(ChannelWithOffset);
+								return FText::FromString(ChannelString);
+							}
+							return FText::GetEmpty();
+						})
+				]
+			]
+		);
+	}
+
+	/** Border to hold contents of the widget */
+	TSharedPtr<SBorder> ContentBorder;
+
+	/** The brush used for the border */
+	FSlateBrush BorderBrush;
+
+	/** The model used by this widget */
+	TSharedPtr<FDMXPixelMappingScreenComponentModel> Model;
+
+	// Slate args
+	int32 ChannelOffset = 0;
+	TAttribute<int32> Universe;
+	TAttribute<int32> StartingChannel;
+};
+
+/** Widget to draw the inners of a screen component */
+class SDMXPixelMappingScreenComponentGrid
+	: public SCompoundWidget
+{
+public:
+	SLATE_BEGIN_ARGS(SDMXPixelMappingScreenComponentGrid)
+	{}
+	SLATE_END_ARGS()
+
+	/** Constructs this widget */
+	void Construct(const FArguments& InArgs, const TSharedRef<FDMXPixelMappingScreenComponentModel>& InModel)
+	{
+		Model = InModel;
+
+		ChildSlot
+			.VAlign(VAlign_Fill)
+			.HAlign(HAlign_Fill)
+			[
+				SAssignNew(GridPanel, SUniformGridPanel)
+			];
+
+		CachedCellFormat = Model->GetCellFormat();
+		CachedDistribution = Model->GetDistribution();
+
+		RebuildGridContent();
+	}
+
+protected:
+	//~ Begin SWidget Interface
+	virtual void Tick(const FGeometry& AllottedGeometry, const double InCurrentTime, const float InDeltaTime) override
+	{
+		if ((GridPanel->GetChildren()->Num() != Model->GetNumColumns() * Model->GetNumRows()) ||
+			CachedDistribution != Model->GetDistribution() ||
+			CachedCellFormat != Model->GetCellFormat())
+		{
+			// Handle cell format changes
+			CachedDistribution = Model->GetDistribution();
+			CachedCellFormat = Model->GetCellFormat();
+
+			RebuildGridContent();
+		}
+	}
+	//~ End SWidget Interface
+
+private:
+	/** Rebuilds the grid content */
+	void RebuildGridContent()
+	{
+		GridPanel->ClearChildren();
+
+		TArray<int32> CellIndices;
+		const int32 TotalCells = Model->GetNumColumns() * Model->GetNumRows();
+		for (int32 CellID = 0; CellID < TotalCells; ++CellID)
+		{
+			CellIndices.Add(CellID);
+		}
+
+		TArray<int32> DistributionSortedCellIndices;
+		FDMXRuntimeUtils::PixelMappingDistributionSort(Model->GetDistribution(), Model->GetNumColumns(), Model->GetNumRows(), CellIndices, DistributionSortedCellIndices);
+
+		const uint32 NumChannelsPerCell = FDMXPixelMappingUtils::GetNumChannelsPerCell(Model->GetCellFormat());
+
+		uint32 XYIndex = 0;
+		for (int32 YIndex = 0; YIndex < Model->GetNumRows(); ++YIndex)
+		{
+			for (int32 XIndex = 0; XIndex < Model->GetNumColumns(); ++XIndex)
+			{
+				if (!ensureMsgf(DistributionSortedCellIndices.IsValidIndex(XYIndex), TEXT("Pixel mapping cell index out of range.")))
+				{
+					break;
+				}
+
+				const int32 ChannelOffset = DistributionSortedCellIndices[XYIndex] * NumChannelsPerCell;
+				GridPanel->AddSlot(XIndex, YIndex)
+					.HAlign(HAlign_Fill)
+					.VAlign(VAlign_Fill)
+					[
+						SNew(SDMXPixelMappingScreenComponentCell, Model.ToSharedRef())
+						.ChannelOffset(ChannelOffset)
+						.StartingChannel_Lambda([this]()
+							{
+								return Model->GetStartingChannel();
+							})
+						.Universe_Lambda([this]()
+							{
+								return Model->GetUniverse();
+							})
+					];
+
+				XYIndex++;
+			}
+		}
+	}
+
+	/** Cell format, cached to detect changes */
+	EDMXCellFormat CachedCellFormat = EDMXCellFormat::PF_RGB;
+
+	/** Pixelmapping distribution, cached to detect changes */
+	EDMXPixelMappingDistribution CachedDistribution = EDMXPixelMappingDistribution::TopLeftToRight;
+
+	/** The actual grid panel */
+	TSharedPtr<SUniformGridPanel> GridPanel;
+
+	/** The model used by this widget */
+	TSharedPtr<FDMXPixelMappingScreenComponentModel> Model;
+};
+
+/** Simple view of a screen component, useful if many cells are drawn */
+class SDMXPixelMappingScreenComponentSimplistic
+	: public SCompoundWidget
+{
+public:
+	SLATE_BEGIN_ARGS(SDMXPixelMappingScreenComponentSimplistic)
+	{}
+	SLATE_END_ARGS()
+
+	//~ Begin SCompoundWidget interface
+	void Construct(const FArguments& InArgs, const TSharedRef<FDMXPixelMappingScreenComponentModel>& InModel)
+	{
+		ChildSlot
+		[
+			SNew(SBox)
+			.HAlign(HAlign_Fill)
+			.VAlign(VAlign_Fill)
+			[
+				SNew(SScaleBox)
+				.Stretch(EStretch::ScaleToFit)
+				.StretchDirection(EStretchDirection::DownOnly)
+				[
+					SNew(SVerticalBox)
+					+ SVerticalBox::Slot()
+					.HAlign(HAlign_Center)
+					.VAlign(VAlign_Center)
+					[
+						SNew(STextBlock)
+						.Font(FCoreStyle::GetDefaultFontStyle("Regular", 24))
+						.ColorAndOpacity(FSlateColor::UseForeground())
+						.Text_Lambda([InModel]()
+							{
+								return FText::Format(LOCTEXT("Num_Pixels", "{0} x {1} pixels"), InModel->GetNumColumns(), InModel->GetNumRows());
+							})
+					]
+					+ SVerticalBox::Slot()
+					.HAlign(HAlign_Center)
+					.VAlign(VAlign_Center)
+					[
+						SNew(STextBlock)
+						.Font(FCoreStyle::GetDefaultFontStyle("Regular", 24))
+						.ColorAndOpacity(FSlateColor::UseForeground())
+						.Text_Lambda([InModel]()
+							{
+								return FText::FromString(FString::Printf(TEXT("%i.%i"), InModel->GetUniverse(), InModel->GetStartingChannel()));
+							})
+					]
+				]
+			]
+		];
+	}
+	//~ Begin SCompoundWidget interface
+
+private:
+	/** The model used by this widget */
+	TSharedPtr<FDMXPixelMappingScreenComponentModel> Model;
+};
+
+void SDMXPixelMappingScreenComponent::Construct(const FArguments& InArgs, const TSharedRef<FDMXPixelMappingToolkit>& InToolkit, TWeakObjectPtr<UDMXPixelMappingScreenComponent> ScreenComponent)
+{
+	if (!ScreenComponent.IsValid())
+	{
+		return;
+	}
+
+	SetCanTick(true);
+
+	WeakToolkit = InToolkit;
+	Model = MakeShared<FDMXPixelMappingScreenComponentModel>(InToolkit, ScreenComponent);
+
+	BorderBrush.DrawAs = ESlateBrushDrawType::Border;
+	BorderBrush.Margin = FMargin(1.f);
+
+	ChildSlot
+	[
+		SNew(SBox)
+		.WidthOverride_Lambda([this]()
+			{
+				return Model->GetSize().X;
+			})
+		.HeightOverride_Lambda([this]()
+			{
+				return Model->GetSize().Y;
+			})
+		[
+			SNew(SOverlay)
+
+			+ SOverlay::Slot() 
+				[
+					SAssignNew(ContentBorder, SBorder)
+					.BorderImage(FAppStyle::GetBrush("NoBorder"))
+					.Padding(0.f)
+				]
+
+			+ SOverlay::Slot() // Required to correctly draw the color around the grid
+			[
+				SNew(SBorder)
+				.BorderImage_Lambda([this]()
+					{
+						BorderBrush.TintColor = Model->GetColor();
+						return &BorderBrush;
+					})
+			]
+		]
+	];	
+
+	UpdateContent();
+}
+
+bool SDMXPixelMappingScreenComponent::Equals(UDMXPixelMappingBaseComponent* Component) const
+{
+	if (Component)
+	{
+		return Model->Equals(Component);
+	}
+	return false;
+}
+
+FVector2D SDMXPixelMappingScreenComponent::GetPosition() const
+{
+	return Model->GetPosition();
+}
+
+void SDMXPixelMappingScreenComponent::Tick(const FGeometry& AllottedGeometry, const double InCurrentTime, const float InDeltaTime)
+{
+	const bool bShouldDrawSimplisticView = Model->GetNumColumns() * Model->GetNumRows() > 160;
+	if (bDrawsSimplisticView != bShouldDrawSimplisticView)
+	{
+		UpdateContent();
+	}
+}
+
+void SDMXPixelMappingScreenComponent::UpdateContent()
+{
+	if(Model->GetNumColumns() * Model->GetNumRows() > 160)
+	{
+		bDrawsSimplisticView = true;
+		ContentBorder->SetContent
+		(
+			SNew(SDMXPixelMappingScreenComponentSimplistic, Model.ToSharedRef())
+		);
+	}
+	else
+	{
+		bDrawsSimplisticView = false;
+		ContentBorder->SetContent
+		(
+			SNew(SDMXPixelMappingScreenComponentGrid, Model.ToSharedRef())
+		);
+	}
+}
+
+#undef LOCTEXT_NAMESPACE
diff --git a/Engine/Plugins/VirtualProduction/DMX/DMXPixelMapping/Source/DMXPixelMappingEditor/Private/Widgets/SDMXPixelMappingOutputComponent.h b/Engine/Plugins/VirtualProduction/DMX/DMXPixelMapping/Source/DMXPixelMappingEditor/Private/Widgets/SDMXPixelMappingOutputComponent.h
new file mode 100644
--- /dev/null
+++ b/Engine/Plugins/VirtualProduction/DMX/DMXPixelMapping/Source/DMXPixelMappingEditor/Private/Widgets/SDMXPixelMappingOutputComponent.h
@@ -0,0 +1,166 @@
+// Copyright Epic Games, Inc. All Rights Reserved.
+
+#pragma once
+
+#include "CoreMinimal.h"
+
+#include "Widgets/SCompoundWidget.h"
+#include "Widgets/Layout/SConstraintCanvas.h"
+
+class FDMXPixelMappingOutputComponentModel;
+class FDMXPixelMappingScreenComponentModel;
+class FDMXPixelMappingToolkit;
+class UDMXPixelMappingBaseComponent;
+class UDMXPixelMappingOutputComponent;
+class UDMXPixelMappingScreenComponent;
+
+class SBorder;
+class SBox;
+class STextBlock;
+
+
+/** Interface for Output Component Widgets */
+class IDMXPixelMappingOutputComponentWidgetInterface
+{
+public:
+	/** Destructor */
+	virtual ~IDMXPixelMappingOutputComponentWidgetInterface();
+
+	/** Adds the widgets to a canvas. If it already resides in a canvas it is removed from that first. */
+	void AddToCanvas(const TSharedRef<SConstraintCanvas>& InCanvas);
+
+	/** Removes the widget from the canvas, if it was added to one. */
+	void RemoveFromCanvas();
+
+	/** Returns the actual widget implementation */
+	virtual TSharedRef<SWidget> AsWidget() = 0;
+
+	/** Returns true if the widet equals the component */
+	virtual bool Equals(UDMXPixelMappingBaseComponent* Component) const = 0;
+
+protected:
+	/** Returns the position of the widget */
+	virtual FVector2D GetPosition() const = 0;
+
+private:
+	/** When added to a parent, the canvas it was added to */
+	TSharedPtr<SConstraintCanvas> ParentCanvas;
+
+	/** The canvas slot of the component widget */
+	SConstraintCanvas::FSlot* Slot = nullptr;
+};
+
+
+/** Widget that draws an Output Component. For Screen Component, see SDMXPixelMappingScreenComponent */
+class SDMXPixelMappingOutputComponent
+	: public IDMXPixelMappingOutputComponentWidgetInterface
+	, public SCompoundWidget
+{
+public:
+	SLATE_BEGIN_ARGS(SDMXPixelMappingOutputComponent)
+	{}
+	SLATE_END_ARGS()
+
+	/** Constructs the widget */
+	void Construct(const FArguments& InArgs, const TSharedRef<FDMXPixelMappingToolkit>& InToolkit, TWeakObjectPtr<UDMXPixelMappingOutputComponent> OutputComponent);
+
+	//~ Begin IDMXPixelMappingOutputComponentWidgetInterface
+	virtual TSharedRef<SWidget> AsWidget() override { return AsShared(); };
+	virtual bool Equals(UDMXPixelMappingBaseComponent* Component) const override;
+protected:
+	virtual FVector2D GetPosition() const override;
+	//~ End IDMXPixelMappingOutputComponentWidgetInterface
+
+private:
+	/** Updates the child slots */
+	void UpdateChildSlots();
+
+	/** Creates the child slot that displays the component name, above the child slot */
+	void CreateComponentNameChildSlotAbove();
+
+	/** Creates the child slot that displays the component name, inside the child slot */
+	void CreateComponentNameChildSlotInside();
+
+	/** Creates the child slot that displays the Fixture ID of the component */
+	void CreateCellIDChildSlot();
+
+	/** Creates the child slot that displays info about the patch such as Addresses or the Fixture ID */
+	void CreatePatchInfoChildSlot();
+
+	/** Brush for the outermost Bborder */
+	FSlateBrush BorderBrush;
+
+	/** The box that is shown */
+	TSharedPtr<SBox> ComponentBox;
+
+	/** The box that is shown */
+	TSharedPtr<STextBlock> IDTextBlock;
+
+	/** Border for content to display names above the widget */
+	TSharedPtr<SBorder> AboveContentBorder;
+
+	/** Border for content to display names in the top row of the widget */
+	TSharedPtr<SBorder> TopContentBorder;
+
+	/** Border for content to display names in the middle of the widget */
+	TSharedPtr<SBorder> MiddleContentBorder;
+
+	/** Border for content to display names in the top row of the widget */
+	TSharedPtr<SBorder> BottomContentBorder;
+
+	/** The model for this widget */
+	TSharedPtr<FDMXPixelMappingOutputComponentModel> Model;
+
+	/** The toolkit that owns */
+	TWeakPtr<FDMXPixelMappingToolkit> WeakToolkit;
+};
+
+/** Widget that draws a Screen Component. */
+class SDMXPixelMappingScreenComponent
+	: public IDMXPixelMappingOutputComponentWidgetInterface
+	, public SCompoundWidget
+{
+public:
+	SLATE_BEGIN_ARGS(SDMXPixelMappingScreenComponent)
+	{}
+	SLATE_END_ARGS()
+
+	/** Constructs the widget */
+	void Construct(const FArguments& InArgs, const TSharedRef<FDMXPixelMappingToolkit>& InToolkit, TWeakObjectPtr<UDMXPixelMappingScreenComponent> ScreenComponent);
+
+	//~ Begin IDMXPixelMappingOutputComponentWidgetInterface
+	virtual TSharedRef<SWidget> AsWidget() override { return AsShared(); };
+	virtual bool Equals(UDMXPixelMappingBaseComponent* Component) const override;
+protected:
+	virtual FVector2D GetPosition() const override;
+	//~ End IDMXPixelMappingOutputComponentWidgetInterface
+
+	//~ Begin SWidget interface
+	virtual void Tick(const FGeometry& AllottedGeometry, const double InCurrentTime, const float InDeltaTime) override;
+	//~ End SWidget interface
+
+private:
+	/** Updates the content of the widget */
+	void UpdateContent();
+
+	/** True if them model currently draws a simplistic view of the component */
+	bool bDrawsSimplisticView = false;
+
+	/** When added to a parent, the canvas it was added to */
+	TSharedPtr<SConstraintCanvas> ParentCanvas;
+
+	/** The canvas slot of the component widget */
+	SConstraintCanvas::FSlot* Slot = nullptr;
+
+	/** Brush for the outermost Bborder */
+	FSlateBrush BorderBrush;
+
+	/** Border that holds the content */
+	TSharedPtr<SBorder> ContentBorder;
+
+	/** The model for this widget */
+	TSharedPtr<FDMXPixelMappingScreenComponentModel> Model;
+
+	/** The toolkit that owns */
+	TWeakPtr<FDMXPixelMappingToolkit> WeakToolkit;
+};
diff --git a/Engine/Plugins/VirtualProduction/DMX/DMXPixelMapping/Source/DMXPixelMappingEditor/Private/Widgets/SDMXPixelMappingTransformHandle.cpp b/Engine/Plugins/VirtualProduction/DMX/DMXPixelMapping/Source/DMXPixelMappingEditor/Private/Widgets/SDMXPixelMappingTransformHandle.cpp
--- a/Engine/Plugins/VirtualProduction/DMX/DMXPixelMapping/Source/DMXPixelMappingEditor/Private/Widgets/SDMXPixelMappingTransformHandle.cpp
+++ b/Engine/Plugins/VirtualProduction/DMX/DMXPixelMapping/Source/DMXPixelMappingEditor/Private/Widgets/SDMXPixelMappingTransformHandle.cpp
@@ -4,7 +4,9 @@
 
 #include "DMXPixelMappingComponentReference.h"
 #include "DMXPixelMappingComponentWidget.h"
+#include "DMXPixelMappingLayoutSettings.h"
 #include "Components/DMXPixelMappingOutputComponent.h"
+#include "Components/DMXPixelMappingMatrixCellComponent.h"
 #include "Toolkits/DMXPixelMappingToolkit.h"
 #include "Views/SDMXPixelMappingDesignerView.h"
 
@@ -89,28 +91,6 @@ FReply SDMXPixelMappingTransformHandle::OnMouseButtonUp(const FGeometry& MyGeome
 			{
 				if (UDMXPixelMappingOutputComponent* ResizedComponent = Cast<UDMXPixelMappingOutputComponent>(ComponentReference.GetComponent()))
 				{
-					// Delete the component if it's no longer over its parent
-					if (ResizedComponent->HasValidParent() &&
-						!ResizedComponent->IsOverParent())
-					{
-						ResizedComponent->SetFlags(RF_Transactional);
-						ResizedComponent->Modify();
-
-						ResizedComponent->GetParent()->RemoveChild(ResizedComponent);
-					}
-
-					// Delete childs if they're no longer over their parent
-					for (UDMXPixelMappingBaseComponent* Child : TArray<UDMXPixelMappingBaseComponent*>(ResizedComponent->Children))
-					{
-						if (UDMXPixelMappingOutputComponent* OutputComponent = Cast<UDMXPixelMappingOutputComponent>(Child))
-						{
-							if (OutputComponent && !OutputComponent->IsOverParent())
-							{
-								ResizedComponent->RemoveChild(Child);
-							}
-						}
-					}
-
 					// Set the final size transacted
 					const FVector2D Delta = MouseEvent.GetScreenSpacePosition() - MouseDownPosition;
 					const FVector2D TranslateAmount = Delta * (1.0f / (DesignerViewWeakPtr.Pin()->GetPreviewScale() * MyGeometry.Scale));
@@ -154,16 +134,15 @@ FReply SDMXPixelMappingTransformHandle::OnMouseMove(const FGeometry& MyGeometry,
 
 void SDMXPixelMappingTransformHandle::Resize(UDMXPixelMappingBaseComponent* BaseComponent, const FVector2D& Direction, const FVector2D& Amount)
 {
-
 	if (UDMXPixelMappingOutputComponent* OutputComponent = Cast<UDMXPixelMappingOutputComponent>(BaseComponent))
 	{
 		FVector2D ComponentSize = OutputComponent->GetSize();
 
 		FMargin Offsets = StartingOffsets;
 
-		FVector2D Movement = Amount * Direction;
-		FVector2D PositionMovement = Movement * (FVector2D(1.0f, 1.0f));
-		FVector2D SizeMovement = Movement;
+		const FVector2D Movement = Amount * Direction;
+		const FVector2D PositionMovement = Movement * (FVector2D(1.0f, 1.0f));
+		const FVector2D SizeMovement = Movement;
 
 		if (Direction.X < 0)
 		{
@@ -179,17 +158,60 @@ void SDMXPixelMappingTransformHandle::Resize(UDMXPixelMappingBaseComponent* Base
 
 		if (Direction.X > 0)
 		{
-			Offsets.Left += (Movement).X;
+			Offsets.Left += Movement.X;
 			Offsets.Right += Amount.X * Direction.X;
 		}
 
 		if (Direction.Y > 0)
 		{
-			Offsets.Top += (Movement).Y;
+			Offsets.Top += Movement.Y;
 			Offsets.Bottom += Amount.Y * Direction.Y;
 		}
 
-		OutputComponent->SetSize(FVector2D(Offsets.Right, Offsets.Bottom).RoundToVector());
+		const FVector2D OldSize = OutputComponent->GetSize();
+		const FVector2D NewSize = FVector2D(Offsets.Right, Offsets.Bottom);
+		if (OldSize == NewSize)
+		{
+			// No unchanged values
+			return;
+		}
+
+		OutputComponent->SetSize(NewSize);
+
+		// Scale children if desired, no division by zero
+		const UDMXPixelMappingLayoutSettings* LayoutSettings = GetDefault<UDMXPixelMappingLayoutSettings>();
+		if (!LayoutSettings ||
+			!LayoutSettings->bScaleChildrenWithParent ||
+			NewSize == FVector2D::ZeroVector)
+		{
+			return;
+		}
+
+		const FVector2D RatioVector = NewSize / OldSize;
+		for (UDMXPixelMappingBaseComponent* BaseChild : OutputComponent->GetChildren())
+		{
+			if (UDMXPixelMappingOutputComponent* Child = Cast<UDMXPixelMappingOutputComponent>(BaseChild))
+			{
+				if (BaseChild->GetClass() == UDMXPixelMappingMatrixCellComponent::StaticClass())
+				{
+					// Don't scale matrix cells, the matrix component already cares for this
+					if (Child->IsLockInDesigner())
+					{
+						continue;
+					}
+				}
+
+				Child->Modify();
+
+				// Scale size (Note, SetSize already clamps)
+				Child->SetSize(Child->GetSize() * RatioVector);
+
+				// Scale position
+				const FVector2D ChildPosition = Child->GetPosition();
+				const FVector2D NewPositionRelative = (ChildPosition - OutputComponent->GetPosition()) * RatioVector;
+				Child->SetPosition(OutputComponent->GetPosition() + NewPositionRelative);
+			}
+		}
 	}
 }
 
diff --git a/Engine/Plugins/VirtualProduction/DMX/DMXPixelMapping/Source/DMXPixelMappingEditor/Private/Widgets/SDMXPixelMappingZoomPan.cpp b/Engine/Plugins/VirtualProduction/DMX/DMXPixelMapping/Source/DMXPixelMappingEditor/Private/Widgets/SDMXPixelMappingZoomPan.cpp
--- a/Engine/Plugins/VirtualProduction/DMX/DMXPixelMapping/Source/DMXPixelMappingEditor/Private/Widgets/SDMXPixelMappingZoomPan.cpp
+++ b/Engine/Plugins/VirtualProduction/DMX/DMXPixelMapping/Source/DMXPixelMappingEditor/Private/Widgets/SDMXPixelMappingZoomPan.cpp
@@ -1,8 +1,12 @@
 // Copyright Epic Games, Inc. All Rights Reserved.
 
 #include "Widgets/SDMXPixelMappingZoomPan.h"
+
 #include "Layout/LayoutUtils.h"
 
+#include "RHI.h"
+
+
 void SDMXPixelMappingZoomPan::Construct(const FArguments& InArgs)
 {
 	bHasRelativeLayoutScale = true;
@@ -32,7 +36,8 @@ void SDMXPixelMappingZoomPan::OnArrangeChildren(const FGeometry& AllottedGeometr
 		const FMargin SlotPadding(ChildSlot.GetPadding());
 		AlignmentArrangeResult XResult = AlignChild<Orient_Horizontal>(AllottedGeometry.Size.X, ChildSlot, SlotPadding, 1);
 		AlignmentArrangeResult YResult = AlignChild<Orient_Vertical>(AllottedGeometry.Size.Y, ChildSlot, SlotPadding, 1);
-
+		
+		const FVector2D Size = ChildSlot.GetWidget()->GetDesiredSize().ClampAxes(0.f, GMaxTextureDimensions);
 		ArrangedChildren.AddWidget( ChildVisibility, AllottedGeometry.MakeChild(
 				ChildSlot.GetWidget(),
 				FVector2D(XResult.Offset, YResult.Offset) - ViewOffset.Get(),
diff --git a/Engine/Plugins/VirtualProduction/DMX/DMXPixelMapping/Source/DMXPixelMappingEditorWidgets/Private/DMXPixelMappingComponentWidget.cpp b/Engine/Plugins/VirtualProduction/DMX/DMXPixelMapping/Source/DMXPixelMappingEditorWidgets/Private/DMXPixelMappingComponentWidget.cpp
--- a/Engine/Plugins/VirtualProduction/DMX/DMXPixelMapping/Source/DMXPixelMappingEditorWidgets/Private/DMXPixelMappingComponentWidget.cpp
+++ b/Engine/Plugins/VirtualProduction/DMX/DMXPixelMapping/Source/DMXPixelMappingEditorWidgets/Private/DMXPixelMappingComponentWidget.cpp
@@ -6,6 +6,7 @@
 #include "SDMXPixelMappingComponentLabel.h"
 
 
+PRAGMA_DISABLE_DEPRECATION_WARNINGS // The whole class and several child widgets are deprecated 5.1
 FDMXPixelMappingComponentWidget::FDMXPixelMappingComponentWidget(TSharedPtr<SDMXPixelMappingComponentBox> InComponentBox, TSharedPtr<SDMXPixelMappingComponentLabel> InComponentLabel, bool bShowLabelAbove)
 	: ComponentBox(InComponentBox)
 	, ComponentLabel(InComponentLabel)
@@ -124,3 +125,4 @@ void FDMXPixelMappingComponentWidget::SetVisibility(EVisibility Visibility)
 	ComponentBox->SetVisibility(Visibility);
 	ComponentLabel->SetVisibility(Visibility);
 }
+PRAGMA_ENABLE_DEPRECATION_WARNINGS
diff --git a/Engine/Plugins/VirtualProduction/DMX/DMXPixelMapping/Source/DMXPixelMappingEditorWidgets/Private/SDMXPixelMappingComponentBox.cpp b/Engine/Plugins/VirtualProduction/DMX/DMXPixelMapping/Source/DMXPixelMappingEditorWidgets/Private/SDMXPixelMappingComponentBox.cpp
--- a/Engine/Plugins/VirtualProduction/DMX/DMXPixelMapping/Source/DMXPixelMappingEditorWidgets/Private/SDMXPixelMappingComponentBox.cpp
+++ b/Engine/Plugins/VirtualProduction/DMX/DMXPixelMapping/Source/DMXPixelMappingEditorWidgets/Private/SDMXPixelMappingComponentBox.cpp
@@ -8,6 +8,7 @@
 #include "Widgets/Text/STextBlock.h"
 
 
+PRAGMA_DISABLE_DEPRECATION_WARNINGS // The whole class is deprecated 5.1
 void SDMXPixelMappingComponentBox::Construct(const FArguments& InArgs)
 {
 	SetCanTick(false);
@@ -87,3 +88,4 @@ FLinearColor SDMXPixelMappingComponentBox::GetBorderColor() const
 {
 	return BorderBrush.TintColor.GetSpecifiedColor();
 }
+PRAGMA_ENABLE_DEPRECATION_WARNINGS
diff --git a/Engine/Plugins/VirtualProduction/DMX/DMXPixelMapping/Source/DMXPixelMappingEditorWidgets/Private/SDMXPixelMappingComponentLabel.cpp b/Engine/Plugins/VirtualProduction/DMX/DMXPixelMapping/Source/DMXPixelMappingEditorWidgets/Private/SDMXPixelMappingComponentLabel.cpp
--- a/Engine/Plugins/VirtualProduction/DMX/DMXPixelMapping/Source/DMXPixelMappingEditorWidgets/Private/SDMXPixelMappingComponentLabel.cpp
+++ b/Engine/Plugins/VirtualProduction/DMX/DMXPixelMapping/Source/DMXPixelMappingEditorWidgets/Private/SDMXPixelMappingComponentLabel.cpp
@@ -10,6 +10,7 @@
 
 #if WITH_EDITOR
 
+PRAGMA_DISABLE_DEPRECATION_WARNINGS // The whole class is deprecated 5.1
 void SDMXPixelMappingComponentLabel::Construct(const FArguments& InArgs)
 {
 	SetCanTick(false);
@@ -79,5 +80,6 @@ const FVector2D& SDMXPixelMappingComponentLabel::GetLocalSize() const
 {
 	return GetCachedGeometry().GetLocalSize();
 }
+PRAGMA_ENABLE_DEPRECATION_WARNINGS
 
 #endif // WITH_EDITOR
diff --git a/Engine/Plugins/VirtualProduction/DMX/DMXPixelMapping/Source/DMXPixelMappingEditorWidgets/Private/SDMXPixelMappingScreenComponentBox.cpp b/Engine/Plugins/VirtualProduction/DMX/DMXPixelMapping/Source/DMXPixelMappingEditorWidgets/Private/SDMXPixelMappingScreenComponentBox.cpp
--- a/Engine/Plugins/VirtualProduction/DMX/DMXPixelMapping/Source/DMXPixelMappingEditorWidgets/Private/SDMXPixelMappingScreenComponentBox.cpp
+++ b/Engine/Plugins/VirtualProduction/DMX/DMXPixelMapping/Source/DMXPixelMappingEditorWidgets/Private/SDMXPixelMappingScreenComponentBox.cpp
@@ -7,6 +7,7 @@
 #include "Widgets/Layout/SBox.h"
 
 
+PRAGMA_DISABLE_DEPRECATION_WARNINGS // The whole class and several child widgets are deprecated 5.1
 const uint32 SDMXPixelMappingScreenComponentBox::MaxGridUICells = 40 * 40;
 
 void SDMXPixelMappingScreenComponentBox::Construct(const FArguments& InArgs)
@@ -51,7 +52,6 @@ void SDMXPixelMappingScreenComponentBox::Construct(const FArguments& InArgs)
 	}
 }
 
-
 void SDMXPixelMappingScreenComponentBox::SetLocalSize(const FVector2D& NewLocalSize)
 {
 	LocalSize = NewLocalSize;
@@ -93,4 +93,5 @@ void SDMXPixelMappingScreenComponentBox::RebuildGrid(const FDMXPixelMappingScree
 			.bShowUniverse(GridParams.bShowUniverse)
 		);
 	}
-}
\ No newline at end of file
+}
+PRAGMA_ENABLE_DEPRECATION_WARNINGS
diff --git a/Engine/Plugins/VirtualProduction/DMX/DMXPixelMapping/Source/DMXPixelMappingEditorWidgets/Public/DMXPixelMappingComponentWidget.h b/Engine/Plugins/VirtualProduction/DMX/DMXPixelMapping/Source/DMXPixelMappingEditorWidgets/Public/DMXPixelMappingComponentWidget.h
--- a/Engine/Plugins/VirtualProduction/DMX/DMXPixelMapping/Source/DMXPixelMappingEditorWidgets/Public/DMXPixelMappingComponentWidget.h
+++ b/Engine/Plugins/VirtualProduction/DMX/DMXPixelMapping/Source/DMXPixelMappingEditorWidgets/Public/DMXPixelMappingComponentWidget.h
@@ -15,6 +15,7 @@ class SConstraintCanvas;
  * Since pixel mapping designer logic relies on computed sizes, two widgets are needed, one as label text, one as the actual box that needs its real size.
  * This widget internally holds both, both are guaranteed to be valid, and adopt sizes and positions reasonably together.
  */
+class UE_DEPRECATED(5.1, "Pixel Mapping Editor Widgets are no longer supported and to be implemented per view. See SDMXPixelMappingOutputComponent for an example.") FDMXPixelMappingComponentWidget;
 class DMXPIXELMAPPINGEDITORWIDGETS_API FDMXPixelMappingComponentWidget
 	: public TSharedFromThis<FDMXPixelMappingComponentWidget>
 {
diff --git a/Engine/Plugins/VirtualProduction/DMX/DMXPixelMapping/Source/DMXPixelMappingEditorWidgets/Public/SDMXPixelMappingComponentBox.h b/Engine/Plugins/VirtualProduction/DMX/DMXPixelMapping/Source/DMXPixelMappingEditorWidgets/Public/SDMXPixelMappingComponentBox.h
--- a/Engine/Plugins/VirtualProduction/DMX/DMXPixelMapping/Source/DMXPixelMappingEditorWidgets/Public/SDMXPixelMappingComponentBox.h
+++ b/Engine/Plugins/VirtualProduction/DMX/DMXPixelMapping/Source/DMXPixelMappingEditorWidgets/Public/SDMXPixelMappingComponentBox.h
@@ -12,6 +12,7 @@ class STextBlock;
 
 
 /** A label component label in pixelmapping designer. Ment to be used with DMXPixelMapingComponentWidgetWrapper */
+class UE_DEPRECATED(5.1, "Pixel Mapping Editor Widgets are no longer supported and to be implemented per view. See SDMXPixelMappingOutputComponent for an example.") SDMXPixelMappingComponentBox;
 class DMXPIXELMAPPINGEDITORWIDGETS_API SDMXPixelMappingComponentBox
 	: public SCompoundWidget
 {
diff --git a/Engine/Plugins/VirtualProduction/DMX/DMXPixelMapping/Source/DMXPixelMappingEditorWidgets/Public/SDMXPixelMappingComponentLabel.h b/Engine/Plugins/VirtualProduction/DMX/DMXPixelMapping/Source/DMXPixelMappingEditorWidgets/Public/SDMXPixelMappingComponentLabel.h
--- a/Engine/Plugins/VirtualProduction/DMX/DMXPixelMapping/Source/DMXPixelMappingEditorWidgets/Public/SDMXPixelMappingComponentLabel.h
+++ b/Engine/Plugins/VirtualProduction/DMX/DMXPixelMapping/Source/DMXPixelMappingEditorWidgets/Public/SDMXPixelMappingComponentLabel.h
@@ -10,6 +10,7 @@ class SBox;
 class STextBlock;
 
 /** A box in pixel mapping designer. Ment to be used with DMXPixelMapingComponentWidgetWrapper */
+class UE_DEPRECATED(5.1, "Pixel Mapping Editor Widgets are no longer supported and to be implemented per view. See SDMXPixelMappingOutputComponent for an example.") SDMXPixelMappingComponentLabel;
 class DMXPIXELMAPPINGEDITORWIDGETS_API SDMXPixelMappingComponentLabel
 	: public SCompoundWidget
 {
diff --git a/Engine/Plugins/VirtualProduction/DMX/DMXPixelMapping/Source/DMXPixelMappingEditorWidgets/Public/SDMXPixelMappingEditorWidgets.h b/Engine/Plugins/VirtualProduction/DMX/DMXPixelMapping/Source/DMXPixelMappingEditorWidgets/Public/SDMXPixelMappingEditorWidgets.h
--- a/Engine/Plugins/VirtualProduction/DMX/DMXPixelMapping/Source/DMXPixelMappingEditorWidgets/Public/SDMXPixelMappingEditorWidgets.h
+++ b/Engine/Plugins/VirtualProduction/DMX/DMXPixelMapping/Source/DMXPixelMappingEditorWidgets/Public/SDMXPixelMappingEditorWidgets.h
@@ -14,6 +14,7 @@
 class STextBlock;
 class SUniformGridPanel;
 
+class UE_DEPRECATED(5.1, "Pixel Mapping Editor Widgets are no longer supported and to be implemented per view. See SDMXPixelMappingOutputComponent for an example.") SDMXPixelMappingScreenLayout;
 class DMXPIXELMAPPINGEDITORWIDGETS_API SDMXPixelMappingScreenLayout 
 	: public SCompoundWidget
 {
@@ -81,6 +82,7 @@ private:
 	TArray<TPair<int32, int32>> SortedList;
 };
 
+class UE_DEPRECATED(5.1, "Pixel Mapping Editor Widgets are no longer supported and to be implemented per view. See SDMXPixelMappingOutputComponent for an example.") SDMXPixelMappingSimpleScreenLayout;
 class DMXPIXELMAPPINGEDITORWIDGETS_API SDMXPixelMappingSimpleScreenLayout
 	: public SCompoundWidget
 {
@@ -120,6 +122,7 @@ private:
 	TAttribute<const FSlateBrush*> Brush;
 };
 
+class UE_DEPRECATED(5.1, "Pixel Mapping Editor Widgets are no longer supported and to be implemented per view. See SDMXPixelMappingOutputComponent for an example.") SDMXPixelMappingCell;
 class DMXPIXELMAPPINGEDITORWIDGETS_API SDMXPixelMappingCell
 	: public SCompoundWidget
 {
diff --git a/Engine/Plugins/VirtualProduction/DMX/DMXPixelMapping/Source/DMXPixelMappingEditorWidgets/Public/SDMXPixelMappingScreenComponentBox.h b/Engine/Plugins/VirtualProduction/DMX/DMXPixelMapping/Source/DMXPixelMappingEditorWidgets/Public/SDMXPixelMappingScreenComponentBox.h
--- a/Engine/Plugins/VirtualProduction/DMX/DMXPixelMapping/Source/DMXPixelMappingEditorWidgets/Public/SDMXPixelMappingScreenComponentBox.h
+++ b/Engine/Plugins/VirtualProduction/DMX/DMXPixelMapping/Source/DMXPixelMappingEditorWidgets/Public/SDMXPixelMappingScreenComponentBox.h
@@ -13,6 +13,7 @@ class SBox;
 
 
 /** Params to update a SDMXPixelMappingScreenComponentBox */
+struct UE_DEPRECATED(5.1, "Pixel Mapping Editor Widgets are no longer supported and to be implemented per view. See SDMXPixelMappingOutputComponent for an example.") FDMXPixelMappingScreenComponentGridParams;
 struct DMXPIXELMAPPINGEDITORWIDGETS_API FDMXPixelMappingScreenComponentGridParams
 {
 	int32 NumXCells = 1;
@@ -25,7 +26,9 @@ struct DMXPIXELMAPPINGEDITORWIDGETS_API FDMXPixelMappingScreenComponentGridParam
 	int32 StartAddress = 1;
 };
 
+PRAGMA_DISABLE_DEPRECATION_WARNINGS
 /** A label component label in pixelmapping designer. Ment to be used with DMXPixelMapingComponentWidgetWrapper */
+class UE_DEPRECATED(5.1, "Pixel Mapping Editor Widgets are no longer supported and to be implemented per view. See SDMXPixelMappingOutputComponent for an example.") SDMXPixelMappingScreenComponentBox;
 class DMXPIXELMAPPINGEDITORWIDGETS_API SDMXPixelMappingScreenComponentBox
 	: public SDMXPixelMappingComponentBox
 {
@@ -83,3 +86,4 @@ private:
 	/** Can't set text, it's not in use in this class */
 	using SDMXPixelMappingComponentBox::SetIDText;
 };
+PRAGMA_ENABLE_DEPRECATION_WARNINGS
diff --git a/Engine/Plugins/VirtualProduction/DMX/DMXPixelMapping/Source/DMXPixelMappingRuntime/Private/Components/DMXPixelMappingFixtureGroupComponent.cpp b/Engine/Plugins/VirtualProduction/DMX/DMXPixelMapping/Source/DMXPixelMappingRuntime/Private/Components/DMXPixelMappingFixtureGroupComponent.cpp
--- a/Engine/Plugins/VirtualProduction/DMX/DMXPixelMapping/Source/DMXPixelMappingRuntime/Private/Components/DMXPixelMappingFixtureGroupComponent.cpp
+++ b/Engine/Plugins/VirtualProduction/DMX/DMXPixelMapping/Source/DMXPixelMappingRuntime/Private/Components/DMXPixelMappingFixtureGroupComponent.cpp
@@ -24,22 +24,21 @@
 
 UDMXPixelMappingFixtureGroupComponent::UDMXPixelMappingFixtureGroupComponent()
 {
-	SizeX = 500.f;
-	SizeY = 500.f;
+	SetSize(FVector2D(500.f, 500.f));
 }
 
 void UDMXPixelMappingFixtureGroupComponent::PostInitProperties()
 {
 	Super::PostInitProperties();
 
-	LastPosition = FVector2D(PositionX, PositionY);
+	LastPosition = GetPosition();
 }
 
 void UDMXPixelMappingFixtureGroupComponent::PostLoad()
 {
 	Super::PostLoad();
 
-	LastPosition = FVector2D(PositionX, PositionY);
+	LastPosition = GetPosition();
 }
 
 #if WITH_EDITOR
@@ -48,26 +47,29 @@ void UDMXPixelMappingFixtureGroupComponent::PostEditChangeProperty(FPropertyChan
 	// Call the parent at the first place
 	Super::PostEditChangeProperty(PropertyChangedEvent);
 
-	if (PropertyChangedEvent.GetPropertyName() == GET_MEMBER_NAME_CHECKED(UDMXPixelMappingOutputComponent, PositionX) ||
-		PropertyChangedEvent.GetPropertyName() == GET_MEMBER_NAME_CHECKED(UDMXPixelMappingOutputComponent, PositionY))
+	if (PropertyChangedEvent.GetPropertyName() == UDMXPixelMappingOutputComponent::GetPositionXPropertyName() ||
+		PropertyChangedEvent.GetPropertyName() == UDMXPixelMappingOutputComponent::GetPositionYPropertyName())
 	{
 		HandlePositionChanged();
 	}
-	if (PropertyChangedEvent.GetPropertyName() == GET_MEMBER_NAME_CHECKED(UDMXPixelMappingOutputComponent, SizeX) ||
-		PropertyChangedEvent.GetPropertyName() == GET_MEMBER_NAME_CHECKED(UDMXPixelMappingOutputComponent, SizeY))
+
+	PRAGMA_DISABLE_DEPRECATION_WARNINGS
+	if (PropertyChangedEvent.GetPropertyName() == UDMXPixelMappingOutputComponent::GetSizeXPropertyName() ||
+		PropertyChangedEvent.GetPropertyName() == UDMXPixelMappingOutputComponent::GetSizeYPropertyName())
 	{
-		if (ComponentWidget.IsValid())
+		if (ComponentWidget_DEPRECATED.IsValid())
 		{
-			ComponentWidget->SetSize(FVector2D(SizeX, SizeY));
+			ComponentWidget_DEPRECATED->SetSize(GetSize());
 		}
 	}
 	else if (PropertyChangedEvent.GetPropertyName() == GET_MEMBER_NAME_CHECKED(UDMXPixelMappingFixtureGroupComponent, DMXLibrary))
 	{
-		if (ComponentWidget.IsValid())
+		if (ComponentWidget_DEPRECATED.IsValid())
 		{
-			ComponentWidget->SetLabelText(FText::FromString(GetUserFriendlyName()));
+			ComponentWidget_DEPRECATED->SetLabelText(FText::FromString(GetUserFriendlyName()));
 		}
 	}
+	PRAGMA_ENABLE_DEPRECATION_WARNINGS
 }
 #endif // WITH_EDITOR
 
@@ -155,39 +157,34 @@ void UDMXPixelMappingFixtureGroupComponent::QueueDownsample()
 
 void UDMXPixelMappingFixtureGroupComponent::SetPosition(const FVector2D& NewPosition)
 {
-	Modify();
-
-	PositionX = FMath::RoundHalfToZero(NewPosition.X);
-	PositionY = FMath::RoundHalfToZero(NewPosition.Y);
+	Super::SetPosition(NewPosition);
 
 	HandlePositionChanged();
 }
 
 void UDMXPixelMappingFixtureGroupComponent::SetSize(const FVector2D& NewSize)
 {
-	Modify();
-
-	SizeX = FMath::RoundHalfToZero(NewSize.X);
-	SizeY = FMath::RoundHalfToZero(NewSize.Y);
-
-	SizeX = FMath::Max(SizeX, 1.f);
-	SizeY = FMath::Max(SizeY, 1.f);
+	Super::SetSize(NewSize);
 
 #if WITH_EDITOR
-	if (ComponentWidget.IsValid())
+	PRAGMA_DISABLE_DEPRECATION_WARNINGS
+	if (ComponentWidget_DEPRECATED.IsValid())
 	{
-		ComponentWidget->SetSize(FVector2D(SizeX, SizeY));
+		ComponentWidget_DEPRECATED->SetSize(GetSize());
 	}
+	PRAGMA_ENABLE_DEPRECATION_WARNINGS
 #endif
 }
 
 void UDMXPixelMappingFixtureGroupComponent::HandlePositionChanged()
 {
 #if WITH_EDITOR
-	if (ComponentWidget.IsValid())
+	PRAGMA_DISABLE_DEPRECATION_WARNINGS
+	if (ComponentWidget_DEPRECATED.IsValid())
 	{
-		ComponentWidget->SetPosition(FVector2D(PositionX, PositionY));
+		ComponentWidget_DEPRECATED->SetPosition(GetPosition());
 	}
+	PRAGMA_ENABLE_DEPRECATION_WARNINGS
 #endif
 
 	// Propagonate to children
diff --git a/Engine/Plugins/VirtualProduction/DMX/DMXPixelMapping/Source/DMXPixelMappingRuntime/Private/Components/DMXPixelMappingFixtureGroupItemComponent.cpp b/Engine/Plugins/VirtualProduction/DMX/DMXPixelMapping/Source/DMXPixelMappingRuntime/Private/Components/DMXPixelMappingFixtureGroupItemComponent.cpp
--- a/Engine/Plugins/VirtualProduction/DMX/DMXPixelMapping/Source/DMXPixelMappingRuntime/Private/Components/DMXPixelMappingFixtureGroupItemComponent.cpp
+++ b/Engine/Plugins/VirtualProduction/DMX/DMXPixelMapping/Source/DMXPixelMappingRuntime/Private/Components/DMXPixelMappingFixtureGroupItemComponent.cpp
@@ -28,8 +28,7 @@ DECLARE_CYCLE_STAT(TEXT("Send Fixture Group Item"), STAT_DMXPixelMaping_FixtureG
 UDMXPixelMappingFixtureGroupItemComponent::UDMXPixelMappingFixtureGroupItemComponent()
 	: DownsamplePixelIndex(0)
 {
-	SizeX = 32.f;
-	SizeY = 32.f;
+	SetSize(FVector2D(32.f, 32.f));
 
 	ColorMode = EDMXColorMode::CM_RGB;
 	AttributeRExpose = AttributeGExpose = AttributeBExpose = true;
@@ -75,22 +74,24 @@ void UDMXPixelMappingFixtureGroupItemComponent::PostEditChangeProperty(FProperty
 
 	FName PropertyName = PropertyChangedEvent.GetPropertyName();
 
-	if (PropertyChangedEvent.GetPropertyName() == GET_MEMBER_NAME_CHECKED(UDMXPixelMappingOutputComponent, PositionX) ||
-		PropertyChangedEvent.GetPropertyName() == GET_MEMBER_NAME_CHECKED(UDMXPixelMappingOutputComponent, PositionY))
+	PRAGMA_DISABLE_DEPRECATION_WARNINGS
+	if (PropertyChangedEvent.GetPropertyName() == UDMXPixelMappingOutputComponent::GetPositionXPropertyName() ||
+		PropertyChangedEvent.GetPropertyName() == UDMXPixelMappingOutputComponent::GetPositionYPropertyName())
 	{
-		if (ComponentWidget.IsValid())
+		if (ComponentWidget_DEPRECATED.IsValid())
 		{
-			ComponentWidget->SetPosition(FVector2D(PositionX, PositionY));
+			ComponentWidget_DEPRECATED->SetPosition(GetPosition());
 		}
 	}
-	if (PropertyChangedEvent.GetPropertyName() == GET_MEMBER_NAME_CHECKED(UDMXPixelMappingOutputComponent, SizeX) ||
-		PropertyChangedEvent.GetPropertyName() == GET_MEMBER_NAME_CHECKED(UDMXPixelMappingOutputComponent, SizeY))
+	if (PropertyChangedEvent.GetPropertyName() == UDMXPixelMappingOutputComponent::GetSizeXPropertyName() ||
+		PropertyChangedEvent.GetPropertyName() == UDMXPixelMappingOutputComponent::GetSizeYPropertyName())
 	{
-		if (ComponentWidget.IsValid())
+		if (ComponentWidget_DEPRECATED.IsValid())
 		{
-			ComponentWidget->SetSize(FVector2D(SizeX, SizeY));
+			ComponentWidget_DEPRECATED->SetSize(GetSize());
 		}
 	}
+	PRAGMA_ENABLE_DEPRECATION_WARNINGS
 }
 #endif // WITH_EDITOR
 
@@ -208,8 +209,8 @@ void UDMXPixelMappingFixtureGroupItemComponent::QueueDownsample()
 	const uint32 TextureSizeY = InputTexture->GetResource()->GetSizeY();
 	check(TextureSizeX > 0 && TextureSizeY > 0);
 	const FIntPoint PixelPosition = RendererComponent->GetPixelPosition(DownsamplePixelIndex);
-	const FVector2D UV = FVector2D(PositionX / TextureSizeX, PositionY / TextureSizeY);
-	const FVector2D UVSize(SizeX / TextureSizeX, SizeY / TextureSizeY);
+	const FVector2D UV = FVector2D(GetPosition().X / TextureSizeX, GetPosition().Y / TextureSizeY);
+	const FVector2D UVSize(GetSize().X / TextureSizeX, GetSize().Y / TextureSizeY);
 	const FVector2D UVCellSize = UVSize / 2.f;
 	constexpr bool bStaticCalculateUV = true;
 
@@ -245,34 +246,29 @@ void UDMXPixelMappingFixtureGroupItemComponent::QueueDownsample()
 
 void UDMXPixelMappingFixtureGroupItemComponent::SetPosition(const FVector2D& NewPosition)
 {
-	Modify();
-
-	PositionX = FMath::RoundHalfToZero(NewPosition.X);
-	PositionY = FMath::RoundHalfToZero(NewPosition.Y);
+	Super::SetPosition(NewPosition);
 
 #if WITH_EDITOR
-	if (ComponentWidget.IsValid())
+	PRAGMA_DISABLE_DEPRECATION_WARNINGS
+	if (ComponentWidget_DEPRECATED.IsValid())
 	{
-		ComponentWidget->SetPosition(FVector2D(PositionX, PositionY));
+		ComponentWidget_DEPRECATED->SetPosition(GetPosition());
 	}
+	PRAGMA_ENABLE_DEPRECATION_WARNINGS
 #endif
 }
 
 void UDMXPixelMappingFixtureGroupItemComponent::SetSize(const FVector2D& NewSize)
 {
-	Modify();
-
-	SizeX = FMath::RoundHalfToZero(NewSize.X);
-	SizeY = FMath::RoundHalfToZero(NewSize.Y);
-
-	SizeX = FMath::Max(SizeX, 1.f);
-	SizeY = FMath::Max(SizeY, 1.f);
+	Super::SetSize(NewSize);
 
 #if WITH_EDITOR
-	if (ComponentWidget.IsValid())
+	PRAGMA_DISABLE_DEPRECATION_WARNINGS
+	if (ComponentWidget_DEPRECATED.IsValid())
 	{
-		ComponentWidget->SetSize(FVector2D(SizeX, SizeY));
+		ComponentWidget_DEPRECATED->SetSize(GetSize());
 	}
+	PRAGMA_ENABLE_DEPRECATION_WARNINGS
 #endif
 }
 
@@ -281,11 +277,21 @@ bool UDMXPixelMappingFixtureGroupItemComponent::IsOverParent() const
 	// Needs be over the over the group
 	if (UDMXPixelMappingFixtureGroupComponent* ParentFixtureGroupComponent = Cast<UDMXPixelMappingFixtureGroupComponent>(GetParent()))
 	{
+		const float Left = GetPosition().X;
+		const float Top = GetPosition().Y;
+		const float Right = GetPosition().X + GetSize().X;
+		const float Bottom = GetPosition().Y + GetSize().Y;
+
+		const float ParentLeft = ParentFixtureGroupComponent->GetPosition().X;
+		const float ParentTop = ParentFixtureGroupComponent->GetPosition().Y;
+		const float ParentRight = ParentFixtureGroupComponent->GetPosition().X + ParentFixtureGroupComponent->GetSize().X;
+		const float ParentBottom = ParentFixtureGroupComponent->GetPosition().Y + ParentFixtureGroupComponent->GetSize().Y;
+
 		return
-			PositionX >= ParentFixtureGroupComponent->GetPosition().X &&
-			PositionY >= ParentFixtureGroupComponent->GetPosition().Y &&
-			PositionX + SizeX <= ParentFixtureGroupComponent->GetPosition().X + ParentFixtureGroupComponent->GetSize().X &&
-			PositionY + SizeY <= ParentFixtureGroupComponent->GetPosition().Y + ParentFixtureGroupComponent->GetSize().Y;
+			Left > ParentLeft - .49f &&
+			Top > ParentTop - .49f &&
+			Right < ParentRight + .49f &&
+			Bottom < ParentBottom + .49f;
 	}
 
 	return false;
diff --git a/Engine/Plugins/VirtualProduction/DMX/DMXPixelMapping/Source/DMXPixelMappingRuntime/Private/Components/DMXPixelMappingMatrixCellComponent.cpp b/Engine/Plugins/VirtualProduction/DMX/DMXPixelMapping/Source/DMXPixelMappingRuntime/Private/Components/DMXPixelMappingMatrixCellComponent.cpp
--- a/Engine/Plugins/VirtualProduction/DMX/DMXPixelMapping/Source/DMXPixelMappingRuntime/Private/Components/DMXPixelMappingMatrixCellComponent.cpp
+++ b/Engine/Plugins/VirtualProduction/DMX/DMXPixelMapping/Source/DMXPixelMappingRuntime/Private/Components/DMXPixelMappingMatrixCellComponent.cpp
@@ -32,6 +32,7 @@ UDMXPixelMappingMatrixCellComponent::UDMXPixelMappingMatrixCellComponent()
 {
 #if WITH_EDITORONLY_DATA
 	bLockInDesigner = true;
+	EditorColor = FLinearColor::White.CopyWithNewOpacity(.4f);
 #endif // WITH_EDITORONLY_DATA
 }
 
@@ -43,46 +44,50 @@ void UDMXPixelMappingMatrixCellComponent::PostEditChangeProperty(FPropertyChange
 
 	FName PropertyName = PropertyChangedEvent.GetPropertyName();
 
-	if (PropertyChangedEvent.GetPropertyName() == GET_MEMBER_NAME_CHECKED(UDMXPixelMappingOutputComponent, PositionX) ||
-		PropertyChangedEvent.GetPropertyName() == GET_MEMBER_NAME_CHECKED(UDMXPixelMappingOutputComponent, PositionY))
+	PRAGMA_DISABLE_DEPRECATION_WARNINGS
+	if (PropertyChangedEvent.GetPropertyName() == UDMXPixelMappingOutputComponent::GetPositionXPropertyName() ||
+		PropertyChangedEvent.GetPropertyName() == UDMXPixelMappingOutputComponent::GetPositionYPropertyName())
 	{
-		if (ComponentWidget.IsValid())
+		if (ComponentWidget_DEPRECATED.IsValid())
 		{
-			ComponentWidget->SetPosition(FVector2D(PositionX, PositionY));
+			ComponentWidget_DEPRECATED->SetPosition(GetPosition());
 		}
 	}
-	if (PropertyChangedEvent.GetPropertyName() == GET_MEMBER_NAME_CHECKED(UDMXPixelMappingOutputComponent, SizeX) ||
-		PropertyChangedEvent.GetPropertyName() == GET_MEMBER_NAME_CHECKED(UDMXPixelMappingOutputComponent, SizeY))
+	if (PropertyChangedEvent.GetPropertyName() == UDMXPixelMappingOutputComponent::GetSizeXPropertyName() ||
+		PropertyChangedEvent.GetPropertyName() == UDMXPixelMappingOutputComponent::GetSizeYPropertyName())
 	{
-		if (ComponentWidget.IsValid())
+		if (ComponentWidget_DEPRECATED.IsValid())
 		{
-			ComponentWidget->SetSize(FVector2D(SizeX, SizeY));
+			ComponentWidget_DEPRECATED->SetSize(GetSize());
 		}
 	}
 	if (PropertyName == GET_MEMBER_NAME_CHECKED(UDMXPixelMappingMatrixCellComponent, CellID))
 	{
-		if (ComponentWidget.IsValid())
+		if (ComponentWidget_DEPRECATED.IsValid())
 		{
-			ComponentWidget->GetComponentBox()->SetIDText(FText::Format(LOCTEXT("CellID", "{0}"), CellID));
+			ComponentWidget_DEPRECATED->GetComponentBox()->SetIDText(FText::Format(LOCTEXT("CellID", "{0}"), CellID));
 		}
 	}
+	PRAGMA_ENABLE_DEPRECATION_WARNINGS
 }
 #endif // WITH_EDITOR
 
 #if WITH_EDITOR
+PRAGMA_DISABLE_DEPRECATION_WARNINGS
 TSharedRef<FDMXPixelMappingComponentWidget> UDMXPixelMappingMatrixCellComponent::BuildSlot(TSharedRef<SConstraintCanvas> InCanvas)
 {
-	ComponentWidget = Super::BuildSlot(InCanvas);
+	ComponentWidget_DEPRECATED = Super::BuildSlot(InCanvas);
 
 	// Expect super to construct the component widget
-	if (ensureMsgf(ComponentWidget.IsValid(), TEXT("PixelMapping: Expected Super to construct a component widget, but didn't.")))
+	if (ensureMsgf(ComponentWidget_DEPRECATED.IsValid(), TEXT("PixelMapping: Expected Super to construct a component widget, but didn't.")))
 	{
-		ComponentWidget->GetComponentLabel()->SetText(FText::GetEmpty());
-		ComponentWidget->GetComponentBox()->SetIDText(FText::Format(LOCTEXT("CellID", "{0}"), CellID));
+		ComponentWidget_DEPRECATED->GetComponentLabel()->SetText(FText::GetEmpty());
+		ComponentWidget_DEPRECATED->GetComponentBox()->SetIDText(FText::Format(LOCTEXT("CellID", "{0}"), CellID));
 	}
 
-	return ComponentWidget.ToSharedRef();
+	return ComponentWidget_DEPRECATED.ToSharedRef();
 }
+PRAGMA_ENABLE_DEPRECATION_WARNINGS
 #endif // WITH_EDITOR
 
 #if WITH_EDITOR
@@ -104,6 +109,7 @@ bool UDMXPixelMappingMatrixCellComponent::IsVisible() const
 #if WITH_EDITOR
 FLinearColor UDMXPixelMappingMatrixCellComponent::GetEditorColor() const
 {
+	PRAGMA_DISABLE_DEPRECATION_WARNINGS
 	if (bLockInDesigner)
 	{
 		// When locked in designer, always use the parent color. So when the parent shows an error color, show it too.
@@ -115,6 +121,7 @@ FLinearColor UDMXPixelMappingMatrixCellComponent::GetEditorColor() const
 			}
 		}
 	}
+	PRAGMA_ENABLE_DEPRECATION_WARNINGS
 
 	return EditorColor;
 }
@@ -123,51 +130,60 @@ FLinearColor UDMXPixelMappingMatrixCellComponent::GetEditorColor() const
 bool UDMXPixelMappingMatrixCellComponent::IsOverParent() const
 {
 	// Needs be over the matrix and over the group
-	if (UDMXPixelMappingMatrixComponent* ParentMatrixComponent = Cast<UDMXPixelMappingMatrixComponent>(GetParent()))
+	UDMXPixelMappingMatrixComponent* ParentMatrixComponent = Cast<UDMXPixelMappingMatrixComponent>(GetParent());
+	if (!ParentMatrixComponent)
+	{
+		return false;
+	}
+	
+	const bool bIsParentMatrixOverGroup = ParentMatrixComponent->IsOverParent();
+	if (!bIsParentMatrixOverGroup)
 	{
-		const bool bIsParentMatrixOverGroup = ParentMatrixComponent->IsOverParent();
-
-		return
-			bIsParentMatrixOverGroup &&
-			PositionX >= ParentMatrixComponent->GetPosition().X &&
-			PositionY >= ParentMatrixComponent->GetPosition().Y &&
-			PositionX + SizeX <= ParentMatrixComponent->GetPosition().X + ParentMatrixComponent->GetSize().X &&
-			PositionY + SizeY <= ParentMatrixComponent->GetPosition().Y + ParentMatrixComponent->GetSize().Y;
+		return false;
 	}
 
-	return false;
+	const float Left = GetPosition().X;
+	const float Top = GetPosition().Y;
+	const float Right = GetPosition().X + GetSize().X;
+	const float Bottom = GetPosition().Y + GetSize().Y;
+
+	const float ParentLeft = ParentMatrixComponent->GetPosition().X;
+	const float ParentTop = ParentMatrixComponent->GetPosition().Y;
+	const float ParentRight = ParentMatrixComponent->GetPosition().X + ParentMatrixComponent->GetSize().X;
+	const float ParentBottom = ParentMatrixComponent->GetPosition().Y + ParentMatrixComponent->GetSize().Y;
+
+	return
+		Left > ParentLeft - .49f &&
+		Top > ParentTop - .49f &&
+		Right < ParentRight + .49f &&
+		Bottom < ParentBottom + .49f;
 }
 
 void UDMXPixelMappingMatrixCellComponent::SetPosition(const FVector2D& NewPosition)
 {
-	Modify();
-
-	PositionX = FMath::RoundHalfToZero(NewPosition.X);
-	PositionY = FMath::RoundHalfToZero(NewPosition.Y);
+	Super::SetPosition(NewPosition);
 
 #if WITH_EDITOR
-	if (ComponentWidget.IsValid())
+	PRAGMA_DISABLE_DEPRECATION_WARNINGS
+	if (ComponentWidget_DEPRECATED.IsValid())
 	{
-		ComponentWidget->SetPosition(FVector2D(PositionX, PositionY));
+		ComponentWidget_DEPRECATED->SetPosition(GetPosition());
 	}
+	PRAGMA_ENABLE_DEPRECATION_WARNINGS
 #endif
 }
 
 void UDMXPixelMappingMatrixCellComponent::SetSize(const FVector2D& NewSize)
 {
-	Modify();
-
-	SizeX = FMath::RoundHalfToZero(NewSize.X);
-	SizeY = FMath::RoundHalfToZero(NewSize.Y);
-
-	SizeX = FMath::Max(SizeX, 1.f);
-	SizeY = FMath::Max(SizeY, 1.f);
+	Super::SetSize(NewSize);
 
 #if WITH_EDITOR
-	if (ComponentWidget.IsValid())
+	PRAGMA_DISABLE_DEPRECATION_WARNINGS
+	if (ComponentWidget_DEPRECATED.IsValid())
 	{
-		ComponentWidget->SetSize(FVector2D(SizeX, SizeY));
+		ComponentWidget_DEPRECATED->SetSize(GetSize());
 	}
+	PRAGMA_ENABLE_DEPRECATION_WARNINGS
 #endif
 }
 
@@ -232,8 +248,8 @@ void UDMXPixelMappingMatrixCellComponent::QueueDownsample()
 	const uint32 TextureSizeY = InputTexture->GetResource()->GetSizeY();
 	check(TextureSizeX > 0 && TextureSizeY > 0);
 	const FIntPoint PixelPosition = RendererComponent->GetPixelPosition(DownsamplePixelIndex);
-	const FVector2D UV = FVector2D(PositionX / TextureSizeX, PositionY / TextureSizeY);
-	const FVector2D UVSize(SizeX / TextureSizeX, SizeY / TextureSizeY);
+	const FVector2D UV = FVector2D(GetPosition().X / TextureSizeX, GetPosition().Y / TextureSizeY);
+	const FVector2D UVSize(GetSize().X / TextureSizeX, GetSize().Y / TextureSizeY);
 	const FVector2D UVCellSize = UVSize / 2.f;
 	constexpr bool bStaticCalculateUV = true;
 
