DMX - Add layout options for DMX Pixelmapping, component widgets now use a model view pattern

diff --git a/Engine/Plugins/VirtualProduction/DMX/DMXPixelMapping/Source/DMXPixelMappingRuntime/Private/Components/DMXPixelMappingMatrixComponent.cpp b/Engine/Plugins/VirtualProduction/DMX/DMXPixelMapping/Source/DMXPixelMappingRuntime/Private/Components/DMXPixelMappingMatrixComponent.cpp
--- a/Engine/Plugins/VirtualProduction/DMX/DMXPixelMapping/Source/DMXPixelMappingRuntime/Private/Components/DMXPixelMappingMatrixComponent.cpp
+++ b/Engine/Plugins/VirtualProduction/DMX/DMXPixelMapping/Source/DMXPixelMappingRuntime/Private/Components/DMXPixelMappingMatrixComponent.cpp
@@ -33,8 +33,7 @@
 
 UDMXPixelMappingMatrixComponent::UDMXPixelMappingMatrixComponent()
 {	
-	SizeX = 32.f;
-	SizeY = 32.f;
+	SetSize(FVector2D(32.f, 32.f));
 
 	AttributeR.SetFromName("Red");
 	AttributeG.SetFromName("Green");
@@ -124,6 +123,22 @@ void UDMXPixelMappingMatrixComponent::LogInvalidProperties()
 	}
 }
 
+#if WITH_EDITOR
+void UDMXPixelMappingMatrixComponent::PreEditChange(FProperty* PropertyAboutToChange)
+{
+	if (PropertyAboutToChange)
+	{
+		const FName PropertyName = PropertyAboutToChange->GetFName();
+
+		if (PropertyName == UDMXPixelMappingOutputComponent::GetPositionXPropertyName() ||
+			PropertyName == UDMXPixelMappingOutputComponent::GetPositionYPropertyName())
+		{
+			PreEditChangePosition = GetPosition();
+		}
+	}
+}
+#endif // WITH_EDITOR
+
 #if WITH_EDITOR
 void UDMXPixelMappingMatrixComponent::PostEditChangeProperty(FPropertyChangedEvent& PropertyChangedEvent)
 {
@@ -132,20 +147,28 @@ void UDMXPixelMappingMatrixComponent::PostEditChangeProperty(FPropertyChangedEve
 	// Call the parent at the first place
 	Super::PostEditChangeProperty(PropertyChangedEvent);
 
-	if (PropertyChangedEvent.GetPropertyName() == GET_MEMBER_NAME_CHECKED(UDMXPixelMappingOutputComponent, PositionX) ||
-		PropertyChangedEvent.GetPropertyName() == GET_MEMBER_NAME_CHECKED(UDMXPixelMappingOutputComponent, PositionY))
+	if (PropertyChangedEvent.GetPropertyName() == UDMXPixelMappingOutputComponent::GetPositionXPropertyName() ||
+		PropertyChangedEvent.GetPropertyName() == UDMXPixelMappingOutputComponent::GetPositionYPropertyName())
 	{
+		constexpr bool bUpdateSizeRecursive = false;
+		ForEachChildOfClass<UDMXPixelMappingMatrixCellComponent>([this](UDMXPixelMappingMatrixCellComponent* ChildComponent)
+			{
+				FVector2D RelativePosition = ChildComponent->GetPosition() - PreEditChangePosition;
+				ChildComponent->SetPosition(GetPosition() + RelativePosition);
+
+			}, bUpdateSizeRecursive);
+
 		HandlePositionChanged();
 	}
-	if (PropertyChangedEvent.GetPropertyName() == GET_MEMBER_NAME_CHECKED(UDMXPixelMappingOutputComponent, SizeX) ||
-		PropertyChangedEvent.GetPropertyName() == GET_MEMBER_NAME_CHECKED(UDMXPixelMappingOutputComponent, SizeY) ||
+	if (PropertyChangedEvent.GetPropertyName() == UDMXPixelMappingOutputComponent::GetSizeXPropertyName() ||
+		PropertyChangedEvent.GetPropertyName() == UDMXPixelMappingOutputComponent::GetSizeYPropertyName() ||
 		PropertyChangedEvent.GetPropertyName() == GET_MEMBER_NAME_CHECKED(UDMXPixelMappingMatrixComponent, CellSize))
 	{
 		HandleSizeChanged();
 
 		// Update size again from the new CellSize that results from handling size or matrix changes
-		SizeX = CellSize.X * CoordinateGrid.X;
-		SizeY = CellSize.Y * CoordinateGrid.Y;
+		const FVector2D NewSize(CellSize.X * CoordinateGrid.X, CellSize.Y * CoordinateGrid.Y);
+		SetSize(NewSize);
 	}
 	if (PropertyChangedEvent.GetPropertyName() == GET_MEMBER_NAME_CHECKED(UDMXPixelMappingMatrixComponent, FixturePatchRef) ||
 		PropertyChangedEvent.GetPropertyName() == GET_MEMBER_NAME_CHECKED(UDMXPixelMappingMatrixComponent, CoordinateGrid))
@@ -181,11 +204,21 @@ bool UDMXPixelMappingMatrixComponent::IsOverParent() const
 	// Needs be over the over the group
 	if (UDMXPixelMappingFixtureGroupComponent* ParentFixtureGroupComponent = Cast<UDMXPixelMappingFixtureGroupComponent>(GetParent()))
 	{
+		const float Left = GetPosition().X;
+		const float Top = GetPosition().Y;
+		const float Right = GetPosition().X + GetSize().X;
+		const float Bottom = GetPosition().Y + GetSize().Y;
+
+		const float ParentLeft = ParentFixtureGroupComponent->GetPosition().X;
+		const float ParentTop = ParentFixtureGroupComponent->GetPosition().Y;
+		const float ParentRight = ParentFixtureGroupComponent->GetPosition().X + ParentFixtureGroupComponent->GetSize().X;
+		const float ParentBottom = ParentFixtureGroupComponent->GetPosition().Y + ParentFixtureGroupComponent->GetSize().Y;
+
 		return
-			PositionX >= ParentFixtureGroupComponent->GetPosition().X &&
-			PositionY >= ParentFixtureGroupComponent->GetPosition().Y &&
-			PositionX + SizeX <= ParentFixtureGroupComponent->GetPosition().X + ParentFixtureGroupComponent->GetSize().X &&
-			PositionY + SizeY <= ParentFixtureGroupComponent->GetPosition().Y + ParentFixtureGroupComponent->GetSize().Y;
+			Left > ParentLeft - .49f &&
+			Top > ParentTop - .49f &&
+			Right < ParentRight + .49f &&
+			Bottom < ParentBottom + .49f;
 	}
 
 	return false;
@@ -312,43 +345,42 @@ FString UDMXPixelMappingMatrixComponent::GetUserFriendlyName() const
 
 void UDMXPixelMappingMatrixComponent::SetPosition(const FVector2D& NewPosition)
 {
-	Modify();
+	const FVector2D OldPosition = GetPosition();
+
+	Super::SetPosition(NewPosition);
+
+	constexpr bool bUpdateSizeRecursive = false;
+	ForEachChildOfClass<UDMXPixelMappingMatrixCellComponent>([&OldPosition, &NewPosition](UDMXPixelMappingMatrixCellComponent* ChildComponent)
+		{
+			FVector2D RelativePosition = ChildComponent->GetPosition() - OldPosition;
+			ChildComponent->SetPosition(NewPosition + RelativePosition);
+
+		}, bUpdateSizeRecursive);
 
-	PositionX = FMath::RoundHalfToZero(NewPosition.X);
-	PositionY = FMath::RoundHalfToZero(NewPosition.Y);
 
 	HandlePositionChanged();
 }
 
 void UDMXPixelMappingMatrixComponent::SetSize(const FVector2D& NewSize)
 {
-	Modify();
-
-	SizeX = FMath::RoundHalfToZero(NewSize.X);
-	SizeY = FMath::RoundHalfToZero(NewSize.Y);
+	const FVector2D OldSize = GetSize();
+	Super::SetSize(NewSize);
 
-	SizeX = FMath::Max(SizeX, 1.f);
-	SizeY = FMath::Max(SizeY, 1.f);
-
-	//HandleMatrixChanged();
-	HandleSizeChanged();
+	if (OldSize != NewSize)
+	{
+		HandleSizeChanged();
+	}
 }
 
 void UDMXPixelMappingMatrixComponent::HandlePositionChanged()
 {
 #if WITH_EDITOR
-	// Propagonate to children
-	constexpr bool bUpdateSizeRecursive = false;
-	ForEachChildOfClass<UDMXPixelMappingMatrixCellComponent>([this](UDMXPixelMappingMatrixCellComponent* ChildComponent)
-		{
-			ChildComponent->SetPosition(GetPosition() + FVector2D(CellSize * ChildComponent->GetCellCoordinate()));
-
-		}, bUpdateSizeRecursive);
-
-	if (ComponentWidget.IsValid())
+PRAGMA_DISABLE_DEPRECATION_WARNINGS
+	if (ComponentWidget_DEPRECATED.IsValid())
 	{
-		ComponentWidget->SetPosition(FVector2D(PositionX, PositionY));
+		ComponentWidget_DEPRECATED->SetPosition(GetPosition());
 	}
+PRAGMA_ENABLE_DEPRECATION_WARNINGS
 #endif
 }
 
@@ -357,41 +389,48 @@ void UDMXPixelMappingMatrixComponent::HandleSizeChanged()
 	CellSize = FVector2D::ZeroVector;
 	if (CoordinateGrid.X > 0 && CoordinateGrid.Y > 0)
 	{
-		CellSize = FVector2D(FMath::RoundToFloat(GetSize().X / CoordinateGrid.X), FMath::RoundToFloat(GetSize().Y / CoordinateGrid.Y));
+		CellSize = FVector2D(GetSize().X / CoordinateGrid.X, GetSize().Y / CoordinateGrid.Y);
+	}
 
-	// Propagonate to children
+	if (Children.Num() > 0)
+	{
 		constexpr bool bUpdateSizeRecursive = false;
 		ForEachChildOfClass<UDMXPixelMappingMatrixCellComponent>([this](UDMXPixelMappingMatrixCellComponent* ChildComponent)
-		{
+			{
+#if WITH_EDITOR
+				if (ChildComponent->IsLockInDesigner())
+				{
+					ChildComponent->SetSize(CellSize);
+					ChildComponent->SetPosition(GetPosition() + FVector2D(CellSize * ChildComponent->GetCellCoordinate()));
+				}
+#else
 				ChildComponent->SetSize(CellSize);
 				ChildComponent->SetPosition(GetPosition() + FVector2D(CellSize * ChildComponent->GetCellCoordinate()));
+#endif // WITH_EDITOR
 
 			}, bUpdateSizeRecursive);
 	}
-
-	if (Children.Num() > 0)
-			{
-		// Update size again from the new CellSize
-		SizeX = CellSize.X * CoordinateGrid.X;
-		SizeY = CellSize.Y * CoordinateGrid.Y;
-			}
 	else
 	{
 		// Set the default size
-		SizeX = 150.f;
-		SizeY = 150.f;
+		const FVector2D DefaultSize(32.f, 32.f);
+		SetSize(DefaultSize);
 	}
 
 #if WITH_EDITOR
-	if (ComponentWidget.IsValid())
+PRAGMA_DISABLE_DEPRECATION_WARNINGS
+	if (ComponentWidget_DEPRECATED.IsValid())
 	{
-		ComponentWidget->SetSize(FVector2D(SizeX, SizeY));
+		ComponentWidget_DEPRECATED->SetSize(GetSize());
 	}
+PRAGMA_ENABLE_DEPRECATION_WARNINGS
 #endif
 }
 
 void UDMXPixelMappingMatrixComponent::HandleMatrixChanged()
 {
+	TGuardValue<bool>(bIsUpdatingChildren, true);
+
 	CoordinateGrid = 0;
 
 	// Remove all existing children and rebuild them anew
diff --git a/Engine/Plugins/VirtualProduction/DMX/DMXPixelMapping/Source/DMXPixelMappingRuntime/Private/Components/DMXPixelMappingOutputComponent.cpp b/Engine/Plugins/VirtualProduction/DMX/DMXPixelMapping/Source/DMXPixelMappingRuntime/Private/Components/DMXPixelMappingOutputComponent.cpp
--- a/Engine/Plugins/VirtualProduction/DMX/DMXPixelMapping/Source/DMXPixelMappingRuntime/Private/Components/DMXPixelMappingOutputComponent.cpp
+++ b/Engine/Plugins/VirtualProduction/DMX/DMXPixelMapping/Source/DMXPixelMappingRuntime/Private/Components/DMXPixelMappingOutputComponent.cpp
@@ -16,8 +16,6 @@
 #define LOCTEXT_NAMESPACE "DMXPixelMappingOutputComponent"
 
 
-const FLinearColor FDMXOutputComponentColors::SelectedColor = FLinearColor::Green;
-
 UDMXPixelMappingOutputComponent::UDMXPixelMappingOutputComponent()
 	: CellBlendingQuality(EDMXPixelBlendingQuality::Low)
 	, PositionX(0.f)
@@ -38,13 +36,15 @@ void UDMXPixelMappingOutputComponent::PostEditChangeProperty(FPropertyChangedEve
 	// Call the parent at the first place
 	Super::PostEditChangeProperty(PropertyChangedEvent);
 	
+PRAGMA_DISABLE_DEPRECATION_WARNINGS
+	// Handle deprecated component widget changes
 	if (PropertyChangedEvent.GetPropertyName() == GET_MEMBER_NAME_CHECKED(UDMXPixelMappingOutputComponent, bVisibleInDesigner))
 	{
 		const EVisibility NewVisiblity = IsVisible() ? EVisibility::HitTestInvisible : EVisibility::Hidden;
 
-		if (ComponentWidget.IsValid())
+		if (ComponentWidget_DEPRECATED.IsValid())
 		{
-			ComponentWidget->SetVisibility(NewVisiblity);
+			ComponentWidget_DEPRECATED->SetVisibility(NewVisiblity);
 		}
 
 		constexpr bool bSetVisibilityRecursive = true;
@@ -56,7 +56,10 @@ void UDMXPixelMappingOutputComponent::PostEditChangeProperty(FPropertyChangedEve
 				}
 			}, bSetVisibilityRecursive);
 	}
-	else if (PropertyChangedEvent.GetPropertyName() == GET_MEMBER_NAME_CHECKED(UDMXPixelMappingOutputComponent, CellBlendingQuality))
+PRAGMA_ENABLE_DEPRECATION_WARNINGS
+
+	// Apply cell belending quality
+	if (PropertyChangedEvent.GetPropertyName() == GET_MEMBER_NAME_CHECKED(UDMXPixelMappingOutputComponent, CellBlendingQuality))
 	{
 		// Propagonate to children
 		constexpr bool bSetCellBlendingQualityToChildsRecursive = true;
@@ -65,45 +68,6 @@ void UDMXPixelMappingOutputComponent::PostEditChangeProperty(FPropertyChangedEve
 				ChildComponent->CellBlendingQuality = CellBlendingQuality;
 			}, bSetCellBlendingQualityToChildsRecursive);
 	}
-
-	if (PropertyChangedEvent.ChangeType != EPropertyChangeType::Interactive)
-	{
-		if (PropertyChangedEvent.GetPropertyName() == GET_MEMBER_NAME_CHECKED(UDMXPixelMappingOutputComponent, PositionX) ||
-			PropertyChangedEvent.GetPropertyName() == GET_MEMBER_NAME_CHECKED(UDMXPixelMappingOutputComponent, PositionY) ||
-			PropertyChangedEvent.GetPropertyName() == GET_MEMBER_NAME_CHECKED(UDMXPixelMappingOutputComponent, SizeX) ||
-			PropertyChangedEvent.GetPropertyName() == GET_MEMBER_NAME_CHECKED(UDMXPixelMappingOutputComponent, SizeY))
-		{
-			PositionX = FMath::RoundHalfToZero(PositionX);
-			PositionY = FMath::RoundHalfToZero(PositionY);
-
-			SizeX = FMath::RoundHalfToZero(SizeX);
-			SizeY = FMath::RoundHalfToZero(SizeY);
-
-			// Remove self if not over parent
-			if (!IsOverParent())
-			{
-				Modify();
-
-				if (ensureMsgf(HasValidParent(), TEXT("No valid Parent when trying to remove PixelMapping Component %s."), *GetUserFriendlyName()))
-				{
-					GetParent()->Modify();
-					GetParent()->RemoveChild(this);
-				}
-			}
-
-			// Remove childs if not over self
-			for (UDMXPixelMappingBaseComponent* Child : TArray<UDMXPixelMappingBaseComponent*>(Children))
-			{
-				if (UDMXPixelMappingOutputComponent* OutputComponent = Cast<UDMXPixelMappingOutputComponent>(Child))
-				{
-					if (OutputComponent && !OutputComponent->IsOverParent())
-					{
-						RemoveChild(OutputComponent);
-					}
-				}
-			}
-		}
-	}
 }
 #endif // WITH_EDITOR
 
@@ -163,13 +127,15 @@ void UDMXPixelMappingOutputComponent::BeginDestroy()
 	Super::BeginDestroy();
 
 #if WITH_EDITOR
-	if (ComponentWidget.IsValid())
+	PRAGMA_DISABLE_DEPRECATION_WARNINGS
+	if (ComponentWidget_DEPRECATED.IsValid())
 	{
-		ComponentWidget->RemoveFromCanvas();
+		ComponentWidget_DEPRECATED->RemoveFromCanvas();
 
 		// Should have released all references by now
-		ensureMsgf(ComponentWidget.GetSharedReferenceCount() == 1, TEXT("Detected Reference to Component Widget the component is destroyed."));
+		ensureMsgf(ComponentWidget_DEPRECATED.GetSharedReferenceCount() == 1, TEXT("Detected Reference to Component Widget the component is destroyed."));
 	}
+	PRAGMA_ENABLE_DEPRECATION_WARNINGS
 #endif
 }
 
@@ -192,15 +158,17 @@ void UDMXPixelMappingOutputComponent::NotifyRemovedFromParent()
 	Super::NotifyRemovedFromParent();
 
 #if WITH_EDITOR
-	if (ComponentWidget.IsValid())
+	PRAGMA_DISABLE_DEPRECATION_WARNINGS
+	if (ComponentWidget_DEPRECATED.IsValid())
 	{
-		ComponentWidget->RemoveFromCanvas();
+		ComponentWidget_DEPRECATED->RemoveFromCanvas();
 
 		// Should have released all references by now
-		ensureMsgf(ComponentWidget.GetSharedReferenceCount() == 1, TEXT("Detected Reference to Component Widget the component is destroyed."));
+		ensureMsgf(ComponentWidget_DEPRECATED.GetSharedReferenceCount() == 1, TEXT("Detected Reference to Component Widget the component is destroyed."));
 
-		ComponentWidget.Reset();
+		ComponentWidget_DEPRECATED.Reset();
 	}
+	PRAGMA_ENABLE_DEPRECATION_WARNINGS
 #endif // WITH_EDITOR
 }
 
@@ -214,17 +182,19 @@ const FText UDMXPixelMappingOutputComponent::GetPaletteCategory()
 #endif // WITH_EDITOR
 
 #if WITH_EDITOR
+PRAGMA_DISABLE_DEPRECATION_WARNINGS
 TSharedRef<FDMXPixelMappingComponentWidget> UDMXPixelMappingOutputComponent::BuildSlot(TSharedRef<SConstraintCanvas> InCanvas)
 {
-	ComponentWidget = MakeShared<FDMXPixelMappingComponentWidget>();
-	ComponentWidget->AddToCanvas(InCanvas, ZOrder);
+	ComponentWidget_DEPRECATED = MakeShared<FDMXPixelMappingComponentWidget>();
+	ComponentWidget_DEPRECATED->AddToCanvas(InCanvas, ZOrder);
 
 	EVisibility NewVisibility = IsVisible() ? EVisibility::HitTestInvisible : EVisibility::Hidden;
 
 	UpdateComponentWidget(NewVisibility);
 
-	return ComponentWidget.ToSharedRef();
+	return ComponentWidget_DEPRECATED.ToSharedRef();
 }
+PRAGMA_ENABLE_DEPRECATION_WARNINGS
 #endif // WITH_EDITOR
 
 #if WITH_EDITOR
@@ -243,10 +213,12 @@ void UDMXPixelMappingOutputComponent::SetZOrder(int32 NewZOrder)
 	
 	ZOrder = NewZOrder;
 
-	// Apply to the UI
-	if (ComponentWidget.IsValid())
+
+PRAGMA_DISABLE_DEPRECATION_WARNINGS
+	// Handle deprecated component widget changes
+	if (ComponentWidget_DEPRECATED.IsValid())
 	{
-		ComponentWidget->SetZOrder(ZOrder);
+		ComponentWidget_DEPRECATED->SetZOrder(ZOrder);
 	}
 
 	constexpr bool bRecursive = true;
@@ -255,19 +227,19 @@ void UDMXPixelMappingOutputComponent::SetZOrder(int32 NewZOrder)
 			if (UDMXPixelMappingOutputComponent* ChildOutputComponent = Cast<UDMXPixelMappingOutputComponent>(ChildComponent))
 			{
 				const int32 NewChildZOrder = ChildOutputComponent->GetZOrder() + DeltaZOrder;
-				if (TSharedPtr<FDMXPixelMappingComponentWidget> ChildComponentWidget = ChildOutputComponent->GetComponentWidget())
+				if (TSharedPtr<FDMXPixelMappingComponentWidget> ChildComponentWidget_DEPRECATED = ChildOutputComponent->GetComponentWidget())
 				{
 					ChildOutputComponent->SetZOrder(NewChildZOrder);
 
 					// Apply to the UI
-					if (ChildOutputComponent->ComponentWidget.IsValid())
+					if (ChildOutputComponent->ComponentWidget_DEPRECATED.IsValid())
 					{
-						ChildOutputComponent->ComponentWidget->SetZOrder(ZOrder);
+						ChildOutputComponent->ComponentWidget_DEPRECATED->SetZOrder(ZOrder);
 					}
 				}
 			}
 		}, bRecursive);
-
+PRAGMA_ENABLE_DEPRECATION_WARNINGS
 }
 #endif // WITH_EDITOR
 
@@ -305,12 +277,20 @@ bool UDMXPixelMappingOutputComponent::OverlapsComponent(UDMXPixelMappingOutputCo
 
 void UDMXPixelMappingOutputComponent::SetPosition(const FVector2D& Position) 
 {
-	ensureMsgf(0, TEXT("SetPosition needs to be implemented in child classes."));
+	if (Position.X != PositionX || Position.Y != PositionY)
+	{
+		PositionX = Position.X;
+		PositionY = Position.Y;
+	}
 }
 
 void UDMXPixelMappingOutputComponent::SetSize(const FVector2D& Size) 
 {
-	ensureMsgf(0, TEXT("SetSize needs to be implemented in child classes."));
+	SizeX = Size.X;
+	SizeY = Size.Y;
+
+	SizeX = FMath::Max(SizeX, 1.f);
+	SizeY = FMath::Max(SizeY, 1.f);
 }
 
 UDMXPixelMappingRendererComponent* UDMXPixelMappingOutputComponent::FindRendererComponent() const
@@ -369,13 +349,14 @@ void UDMXPixelMappingOutputComponent::MakeHighestZOrderInComponentRect()
 #if WITH_EDITOR
 void UDMXPixelMappingOutputComponent::UpdateComponentWidget(EVisibility NewVisibility, bool bWithChildrenRecursive)
 {
-	if (ComponentWidget.IsValid())
+	PRAGMA_DISABLE_DEPRECATION_WARNINGS
+	if (ComponentWidget_DEPRECATED.IsValid())
 	{
-		ComponentWidget->SetVisibility(NewVisibility);
-		ComponentWidget->SetPosition(GetPosition());
-		ComponentWidget->SetSize(GetSize());
-		ComponentWidget->SetColor(GetEditorColor());
-		ComponentWidget->SetLabelText(FText::FromString(GetUserFriendlyName()));
+		ComponentWidget_DEPRECATED->SetPosition(FVector2D(PositionX, PositionY));
+		ComponentWidget_DEPRECATED->SetSize(FVector2D(SizeX, SizeY));
+		ComponentWidget_DEPRECATED->SetVisibility(NewVisibility);
+		ComponentWidget_DEPRECATED->SetColor(GetEditorColor());
+		ComponentWidget_DEPRECATED->SetLabelText(FText::FromString(GetUserFriendlyName()));
 	}
 
 	if (bWithChildrenRecursive)
@@ -389,6 +370,7 @@ void UDMXPixelMappingOutputComponent::UpdateComponentWidget(EVisibility NewVisib
 			}
 		}
 	}
+	PRAGMA_ENABLE_DEPRECATION_WARNINGS
 }
 #endif // WITH_EDITOR
 
diff --git a/Engine/Plugins/VirtualProduction/DMX/DMXPixelMapping/Source/DMXPixelMappingRuntime/Private/Components/DMXPixelMappingRendererComponent.cpp b/Engine/Plugins/VirtualProduction/DMX/DMXPixelMapping/Source/DMXPixelMappingRuntime/Private/Components/DMXPixelMappingRendererComponent.cpp
--- a/Engine/Plugins/VirtualProduction/DMX/DMXPixelMapping/Source/DMXPixelMappingRuntime/Private/Components/DMXPixelMappingRendererComponent.cpp
+++ b/Engine/Plugins/VirtualProduction/DMX/DMXPixelMapping/Source/DMXPixelMappingRuntime/Private/Components/DMXPixelMappingRendererComponent.cpp
@@ -45,8 +45,7 @@ UDMXPixelMappingRendererComponent::UDMXPixelMappingRendererComponent()
 	RendererType = EDMXPixelMappingRendererType::Texture;
 #endif
 	
-	SizeX = 100.f;
-	SizeY = 100.f;
+	SetSize(FVector2D(100.f, 100.f));
 
 	Brightness = 1.0f;
 
@@ -125,22 +124,24 @@ void UDMXPixelMappingRendererComponent::PostEditChangeChainProperty(FPropertyCha
 	// Call the parent at the first place
 	Super::PostEditChangeChainProperty(PropertyChangedChainEvent);
 
-	if (PropertyChangedChainEvent.GetPropertyName() == GET_MEMBER_NAME_CHECKED(UDMXPixelMappingOutputComponent, SizeX) ||
-		PropertyChangedChainEvent.GetPropertyName() == GET_MEMBER_NAME_CHECKED(UDMXPixelMappingOutputComponent, SizeY))
+	if (PropertyChangedChainEvent.GetPropertyName() == UDMXPixelMappingOutputComponent::GetSizeXPropertyName() ||
+		PropertyChangedChainEvent.GetPropertyName() == UDMXPixelMappingOutputComponent::GetSizeYPropertyName())
 	{
 		// The target always needs be within GMaxTextureDimensions, larger dimensions are not supported by the engine
 		const uint32 MaxTextureDimensions = GetMax2DTextureDimension();
 
-		if (SizeX > MaxTextureDimensions ||
-			SizeY > MaxTextureDimensions)
+		if (GetSize().X > MaxTextureDimensions ||
+			GetSize().Y > MaxTextureDimensions)
 		{
-			SizeX = FMath::Clamp(SizeX, 0.0f, static_cast<float>(MaxTextureDimensions));
-			SizeY = FMath::Clamp(SizeY, 0.0f, static_cast<float>(MaxTextureDimensions));
+			const float NewSizeX = FMath::Clamp(GetSize().X, 0.0f, static_cast<float>(MaxTextureDimensions));
+			const float NewSizeY = FMath::Clamp(GetSize().Y, 0.0f, static_cast<float>(MaxTextureDimensions));
+			const FVector2D NewSize(NewSizeX, NewSizeY);
+			SetSize(NewSize);
 
 			UE_LOG(LogDMXPixelMappingRuntime, Warning, TEXT("Pixel mapping textures are limited to engine's max texture dimension %dx%d"), MaxTextureDimensions, MaxTextureDimensions);
 		}
 
-		ResizeMaterialRenderTarget(SizeX, SizeY);
+		ResizeMaterialRenderTarget(GetSize().X, GetSize().Y);
 	} 
 	else if (PropertyChangedChainEvent.GetPropertyName() == GET_MEMBER_NAME_CHECKED(UDMXPixelMappingRendererComponent, RendererType))
 	{
@@ -282,6 +283,7 @@ TSharedRef<SWidget> UDMXPixelMappingRendererComponent::TakeWidget()
 			SNew(SConstraintCanvas);
 	}
 
+	PRAGMA_DISABLE_DEPRECATION_WARNINGS
 	ForEachChild([&](UDMXPixelMappingBaseComponent* InComponent) {
 		if (UDMXPixelMappingOutputComponent* Component = Cast<UDMXPixelMappingOutputComponent>(InComponent))
 		{
@@ -289,6 +291,7 @@ TSharedRef<SWidget> UDMXPixelMappingRendererComponent::TakeWidget()
 			Component->BuildSlot(ComponentsCanvas.ToSharedRef());
 		}
 	}, true);
+	PRAGMA_ENABLE_DEPRECATION_WARNINGS
 
 	return ComponentsCanvas.ToSharedRef();
 }
@@ -388,12 +391,6 @@ void UDMXPixelMappingRendererComponent::RenderAndSendDMX()
 	SendDMX();
 }
 
-void UDMXPixelMappingRendererComponent::SetSize(const FVector2D& NewSize)
-{
-	SizeX = NewSize.X;
-	SizeY = NewSize.Y;
-}
-
 FIntPoint UDMXPixelMappingRendererComponent::GetPixelPosition(int32 InPosition) const
 {
 	const int32 YRows = InPosition / MaxDownsampleBufferTargetSize.X;
@@ -502,7 +499,7 @@ UTextureRenderTarget2D* UDMXPixelMappingRendererComponent::CreateRenderTarget(co
 	UTextureRenderTarget2D* RenderTarget = NewObject<UTextureRenderTarget2D>(this, TargetName);
 	RenderTarget->ClearColor = ClearTextureColor;
 	constexpr bool bInForceLinearGamma = false;
-	RenderTarget->InitCustomFormat(SizeX, SizeY, EPixelFormat::PF_B8G8R8A8, bInForceLinearGamma);
+	RenderTarget->InitCustomFormat(GetSize().X, GetSize().Y, EPixelFormat::PF_B8G8R8A8, bInForceLinearGamma);
 
 	return  RenderTarget;
 }
@@ -538,7 +535,7 @@ void UDMXPixelMappingRendererComponent::RendererInputTexture()
 	}
 	else
 	{
-		ResizePreviewRenderTarget(SizeX, SizeY);
+		ResizePreviewRenderTarget(GetSize().X, GetSize().Y);
 	}
 #endif
 }
diff --git a/Engine/Plugins/VirtualProduction/DMX/DMXPixelMapping/Source/DMXPixelMappingRuntime/Private/Components/DMXPixelMappingScreenComponent.cpp b/Engine/Plugins/VirtualProduction/DMX/DMXPixelMapping/Source/DMXPixelMappingRuntime/Private/Components/DMXPixelMappingScreenComponent.cpp
--- a/Engine/Plugins/VirtualProduction/DMX/DMXPixelMapping/Source/DMXPixelMappingRuntime/Private/Components/DMXPixelMappingScreenComponent.cpp
+++ b/Engine/Plugins/VirtualProduction/DMX/DMXPixelMapping/Source/DMXPixelMappingRuntime/Private/Components/DMXPixelMappingScreenComponent.cpp
@@ -26,9 +26,8 @@ const FVector2D UDMXPixelMappingScreenComponent::MinGridSize = FVector2D(1.f);
 UDMXPixelMappingScreenComponent::UDMXPixelMappingScreenComponent()
 	: bSendToAllOutputPorts(true)
 {
-	SizeX = 500.f; 
-	SizeY = 500.f; 
-
+	SetSize(FVector2D(500.f, 500.f)); 
+	
 	NumXCells = 10;
 	NumYCells = 10;
 
@@ -48,6 +47,25 @@ void UDMXPixelMappingScreenComponent::PostEditChangeChainProperty(FPropertyChang
 	// Call the parent at the first place
 	Super::PostEditChangeChainProperty(PropertyChangedChainEvent);
 	
+	if (PropertyChangedChainEvent.GetPropertyName() == GET_MEMBER_NAME_CHECKED(UDMXPixelMappingScreenComponent, OutputPortReferences))
+	{
+		// Rebuild the set of ports
+		OutputPorts.Reset();
+		for (const FDMXOutputPortReference& OutputPortReference : OutputPortReferences)
+		{
+			const FDMXOutputPortSharedRef* OutputPortPtr = FDMXPortManager::Get().GetOutputPorts().FindByPredicate([&OutputPortReference](const FDMXOutputPortSharedRef& OutputPort) {
+				return OutputPort->GetPortGuid() == OutputPortReference.GetPortGuid();
+				});
+
+			if (OutputPortPtr)
+			{
+				OutputPorts.Add(*OutputPortPtr);
+			}
+		}
+	}
+	
+
+	PRAGMA_DISABLE_DEPRECATION_WARNINGS
 	if (PropertyChangedChainEvent.GetPropertyName() == GET_MEMBER_NAME_CHECKED(UDMXPixelMappingScreenComponent, NumXCells) ||
 		PropertyChangedChainEvent.GetPropertyName() == GET_MEMBER_NAME_CHECKED(UDMXPixelMappingScreenComponent, NumYCells) ||
 		PropertyChangedChainEvent.GetPropertyName() == GET_MEMBER_NAME_CHECKED(UDMXPixelMappingScreenComponent, LocalUniverse) ||
@@ -72,41 +90,26 @@ void UDMXPixelMappingScreenComponent::PostEditChangeChainProperty(FPropertyChang
 			ScreenComponentBox->RebuildGrid(GridParams);
 		}
 	}
-	else if (PropertyChangedChainEvent.GetPropertyName() == GET_MEMBER_NAME_CHECKED(UDMXPixelMappingScreenComponent, OutputPortReferences))
+	if (PropertyChangedChainEvent.ChangeType != EPropertyChangeType::Interactive)
 	{
-		// Rebuild the set of ports
-		OutputPorts.Reset();
-		for (const FDMXOutputPortReference& OutputPortReference : OutputPortReferences)
+		if (PropertyChangedChainEvent.GetPropertyName() == UDMXPixelMappingOutputComponent::GetPositionXPropertyName() ||
+			PropertyChangedChainEvent.GetPropertyName() == UDMXPixelMappingOutputComponent::GetPositionYPropertyName())
 		{
-			const FDMXOutputPortSharedRef* OutputPortPtr = FDMXPortManager::Get().GetOutputPorts().FindByPredicate([&OutputPortReference](const FDMXOutputPortSharedRef& OutputPort) {
-				return OutputPort->GetPortGuid() == OutputPortReference.GetPortGuid();
-				});
-
-			if (OutputPortPtr)
+			if (ComponentWidget_DEPRECATED.IsValid())
 			{
-				OutputPorts.Add(*OutputPortPtr);
+				ComponentWidget_DEPRECATED->SetPosition(GetPosition());
 			}
 		}
-	}
-	else if (PropertyChangedChainEvent.ChangeType != EPropertyChangeType::Interactive)
-	{
-		if (PropertyChangedChainEvent.GetPropertyName() == GET_MEMBER_NAME_CHECKED(UDMXPixelMappingScreenComponent, PositionX) ||
-			PropertyChangedChainEvent.GetPropertyName() == GET_MEMBER_NAME_CHECKED(UDMXPixelMappingScreenComponent, PositionY))
+		else if (PropertyChangedChainEvent.GetPropertyName() == UDMXPixelMappingOutputComponent::GetSizeXPropertyName() ||
+			PropertyChangedChainEvent.GetPropertyName() == UDMXPixelMappingOutputComponent::GetSizeYPropertyName())
 		{
-			if (ComponentWidget.IsValid())
+			if (ComponentWidget_DEPRECATED.IsValid())
 			{
-				ComponentWidget->SetPosition(FVector2D(PositionX, PositionY));
+				ComponentWidget_DEPRECATED->SetSize(GetPosition());
 			}
 		}
-		else if (PropertyChangedChainEvent.GetPropertyName() == GET_MEMBER_NAME_CHECKED(UDMXPixelMappingScreenComponent, SizeX) ||
-			PropertyChangedChainEvent.GetPropertyName() == GET_MEMBER_NAME_CHECKED(UDMXPixelMappingScreenComponent, SizeY))
-		{
-			if (ComponentWidget.IsValid())
-			{
-				ComponentWidget->SetSize(FVector2D(SizeX, SizeY));
-}
-		}
 	}
+	PRAGMA_ENABLE_DEPRECATION_WARNINGS
 }
 #endif // WITH_EDITOR
 
@@ -118,9 +121,10 @@ const FText UDMXPixelMappingScreenComponent::GetPaletteCategory()
 #endif // WITH_EDITOR
 
 #if WITH_EDITOR
+PRAGMA_DISABLE_DEPRECATION_WARNINGS
 TSharedRef<FDMXPixelMappingComponentWidget> UDMXPixelMappingScreenComponent::BuildSlot(TSharedRef<SConstraintCanvas> InCanvas)
 {
-	if (!ComponentWidget.IsValid())
+	if (!ComponentWidget_DEPRECATED.IsValid())
 	{
 		ScreenComponentBox = 
 			SNew(SDMXPixelMappingScreenComponentBox)
@@ -133,17 +137,18 @@ TSharedRef<FDMXPixelMappingComponentWidget> UDMXPixelMappingScreenComponent::Bui
 			.bShowAddresses(bShowAddresses)
 			.bShowUniverse(bShowUniverse);
 
-		ComponentWidget = MakeShared<FDMXPixelMappingComponentWidget>(ScreenComponentBox, nullptr);
+		ComponentWidget_DEPRECATED = MakeShared<FDMXPixelMappingComponentWidget>(ScreenComponentBox, nullptr);
 
-		ComponentWidget->AddToCanvas(InCanvas, ZOrder);
-		ComponentWidget->SetPosition(GetPosition());
-		ComponentWidget->SetSize(GetSize());
-		ComponentWidget->SetColor(GetEditorColor());
-		ComponentWidget->SetLabelText(FText::FromString(GetUserFriendlyName()));
+		ComponentWidget_DEPRECATED->AddToCanvas(InCanvas, ZOrder);
+		ComponentWidget_DEPRECATED->SetPosition(GetPosition());
+		ComponentWidget_DEPRECATED->SetSize(GetSize());
+		ComponentWidget_DEPRECATED->SetColor(GetEditorColor());
+		ComponentWidget_DEPRECATED->SetLabelText(FText::FromString(GetUserFriendlyName()));
 	}
 	
-	return ComponentWidget.ToSharedRef();
+	return ComponentWidget_DEPRECATED.ToSharedRef();
 }
+PRAGMA_ENABLE_DEPRECATION_WARNINGS
 #endif // WITH_EDITOR
 
 const FName& UDMXPixelMappingScreenComponent::GetNamePrefix()
@@ -349,9 +354,9 @@ void UDMXPixelMappingScreenComponent::SendDMX()
 
 			// should be channels from 1...UniverseMaxChannels
 			ChannelToValueMap.Add(SendDMXIndex, SendBuffer[FragmentMapIndex]);
-			if (SendDMXIndex > UniverseMaxChannels || SendDMXIndex < 1)
+			if (!ensureMsgf(SendDMXIndex < UniverseMaxChannels || SendDMXIndex > 0, TEXT("Pixel Mapping Screen Component trying to send out of universe range.")))
 			{
-				UE_LOG(LogTemp, Warning, TEXT("WrongIndex FragmentMapIndex %d, SendDMXIndex %d"), FragmentMapIndex, SendDMXIndex);
+				break;
 			}
 
 			// send dmx if next iteration is the last one
@@ -386,7 +391,7 @@ void UDMXPixelMappingScreenComponent::QueueDownsample()
 	check(TextureSizeX > 0 && TextureSizeY > 0);
 	
 	constexpr bool bStaticCalculateUV = true;
-	const FVector2D SizePixel(SizeX / NumXCells, SizeY / NumYCells);
+	const FVector2D SizePixel(GetSize().X / NumXCells, GetSize().Y / NumYCells);
 	const FVector2D UVSize(SizePixel.X / TextureSizeX, SizePixel.Y / TextureSizeY);
 	const FVector2D UVCellSize = UVSize / 2.f;
 	const int32 PixelNum = NumXCells * NumYCells;
@@ -400,7 +405,7 @@ void UDMXPixelMappingScreenComponent::QueueDownsample()
 	ForEachPixel([&](const int32 InXYIndex, const int32 XIndex, const int32 YIndex)
         {
             const FIntPoint PixelPosition = RendererComponent->GetPixelPosition(InXYIndex + PixelDownsamplePositionRange.Key);
-            const FVector2D UV = FVector2D((PositionX + SizePixel.X * XIndex) / TextureSizeX, (PositionY + SizePixel.Y * YIndex) / TextureSizeY);
+            const FVector2D UV = FVector2D((GetPosition().X + SizePixel.X * XIndex) / TextureSizeX, (GetPosition().Y + SizePixel.Y * YIndex) / TextureSizeY);
 
             FDMXPixelMappingDownsamplePixelParam DownsamplePixelParam
             {
@@ -440,7 +445,7 @@ bool UDMXPixelMappingScreenComponent::CanBeMovedTo(const UDMXPixelMappingBaseCom
 
 const FVector2D UDMXPixelMappingScreenComponent::GetScreenPixelSize() const
 {
-	return FVector2D(SizeX / NumXCells, SizeY / NumYCells);
+	return FVector2D(GetSize().X / NumXCells, GetSize().Y / NumYCells);
 }
 
 void UDMXPixelMappingScreenComponent::ForEachPixel(ForEachPixelCallback InCallback)
@@ -458,30 +463,29 @@ void UDMXPixelMappingScreenComponent::ForEachPixel(ForEachPixelCallback InCallba
 
 void UDMXPixelMappingScreenComponent::SetPosition(const FVector2D& NewPosition)
 {
-	PositionX = NewPosition.X;
-	PositionY = NewPosition.Y;
+	Super::SetPosition(NewPosition);
 
 #if WITH_EDITOR
-	if (ComponentWidget.IsValid())
+	PRAGMA_DISABLE_DEPRECATION_WARNINGS
+	if (ComponentWidget_DEPRECATED.IsValid())
 	{
-		ComponentWidget->SetPosition(FVector2D(PositionX, PositionY));
+		ComponentWidget_DEPRECATED->SetPosition(GetPosition());
 	}
+	PRAGMA_ENABLE_DEPRECATION_WARNINGS
 #endif
 }
 
 void UDMXPixelMappingScreenComponent::SetSize(const FVector2D& NewSize)
 {
-	SizeX = FMath::Max(NewSize.X, 1.f);
-	SizeY = FMath::Max(NewSize.Y, 1.f);
-
-	SizeX = FMath::Max(SizeX, 1.f);
-	SizeY = FMath::Max(SizeY, 1.f);
+	Super::SetSize(NewSize);
 
 #if WITH_EDITOR
-	if (ComponentWidget.IsValid())
+	PRAGMA_DISABLE_DEPRECATION_WARNINGS
+	if (ComponentWidget_DEPRECATED.IsValid())
 	{
-		ComponentWidget->SetSize(FVector2D(SizeX, SizeY));
+		ComponentWidget_DEPRECATED->SetSize(GetSize());
 	}
+	PRAGMA_ENABLE_DEPRECATION_WARNINGS
 #endif
 }
 
diff --git a/Engine/Plugins/VirtualProduction/DMX/DMXPixelMapping/Source/DMXPixelMappingRuntime/Private/DMXPixelMapping.cpp b/Engine/Plugins/VirtualProduction/DMX/DMXPixelMapping/Source/DMXPixelMappingRuntime/Private/DMXPixelMapping.cpp
--- a/Engine/Plugins/VirtualProduction/DMX/DMXPixelMapping/Source/DMXPixelMappingRuntime/Private/DMXPixelMapping.cpp
+++ b/Engine/Plugins/VirtualProduction/DMX/DMXPixelMapping/Source/DMXPixelMappingRuntime/Private/DMXPixelMapping.cpp
@@ -129,6 +129,7 @@ UDMXPixelMappingBaseComponent* UDMXPixelMapping::FindComponent(const FName& InNa
 #if WITH_EDITOR
 UDMXPixelMappingOutputComponent* UDMXPixelMapping::FindComponent(TSharedPtr<SWidget> InWidget) const
 {
+	PRAGMA_DISABLE_DEPRECATION_WARNINGS
 	UDMXPixelMappingOutputComponent* FoundComponent = nullptr;
 
 	ForEachComponentOfClass<UDMXPixelMappingOutputComponent>([&](UDMXPixelMappingOutputComponent* InComponent)
@@ -143,6 +144,7 @@ UDMXPixelMappingOutputComponent* UDMXPixelMapping::FindComponent(TSharedPtr<SWid
 		});
 
 	return FoundComponent;
+	PRAGMA_ENABLE_DEPRECATION_WARNINGS
 }
 
 #endif // WITH_EDITOR
diff --git a/Engine/Plugins/VirtualProduction/DMX/DMXPixelMapping/Source/DMXPixelMappingRuntime/Public/Components/DMXPixelMappingFixtureGroupComponent.h b/Engine/Plugins/VirtualProduction/DMX/DMXPixelMapping/Source/DMXPixelMappingRuntime/Public/Components/DMXPixelMappingFixtureGroupComponent.h
--- a/Engine/Plugins/VirtualProduction/DMX/DMXPixelMapping/Source/DMXPixelMappingRuntime/Public/Components/DMXPixelMappingFixtureGroupComponent.h
+++ b/Engine/Plugins/VirtualProduction/DMX/DMXPixelMapping/Source/DMXPixelMappingRuntime/Public/Components/DMXPixelMappingFixtureGroupComponent.h
@@ -6,11 +6,13 @@
 #include "Library/DMXEntityReference.h"
 #include "DMXPixelMappingFixtureGroupComponent.generated.h"
 
-
 class UDMXLibrary;
+class UDMXPixelMappingLayoutScript;
+
 class STextBlock;
 class SUniformGridPanel;
 
+
 /**
  * Container component for Fixture Items
  */
@@ -54,6 +56,10 @@ public:
 	UPROPERTY(EditAnywhere, BlueprintReadOnly, Category = "Fixture List")
 	TObjectPtr<UDMXLibrary> DMXLibrary;
 
+	/** Layout script for the children of this component (hidden in customizations and displayed in its own panel). */
+	UPROPERTY(EditAnywhere, Instanced, Category = "Layout")
+	TObjectPtr<UDMXPixelMappingLayoutScript> LayoutScript;
+
 private:
 	/** Handles changes in position */
 	void HandlePositionChanged();
diff --git a/Engine/Plugins/VirtualProduction/DMX/DMXPixelMapping/Source/DMXPixelMappingRuntime/Public/Components/DMXPixelMappingMatrixCellComponent.h b/Engine/Plugins/VirtualProduction/DMX/DMXPixelMapping/Source/DMXPixelMappingRuntime/Public/Components/DMXPixelMappingMatrixCellComponent.h
--- a/Engine/Plugins/VirtualProduction/DMX/DMXPixelMapping/Source/DMXPixelMappingRuntime/Public/Components/DMXPixelMappingMatrixCellComponent.h
+++ b/Engine/Plugins/VirtualProduction/DMX/DMXPixelMapping/Source/DMXPixelMappingRuntime/Public/Components/DMXPixelMappingMatrixCellComponent.h
@@ -43,7 +43,9 @@ public:
 
 	//~ Begin UDMXPixelMappingOutputComponent implementation
 #if WITH_EDITOR
+	PRAGMA_DISABLE_DEPRECATION_WARNINGS // FDMXPixelMappingComponentWidget is deprecated
 	virtual TSharedRef<FDMXPixelMappingComponentWidget> BuildSlot(TSharedRef<SConstraintCanvas> InCanvas) override;
+	PRAGMA_ENABLE_DEPRECATION_WARNINGS
 	virtual bool IsVisible() const override;
 	virtual FLinearColor GetEditorColor() const override;
 #endif // WITH_EDITOR	
diff --git a/Engine/Plugins/VirtualProduction/DMX/DMXPixelMapping/Source/DMXPixelMappingRuntime/Public/Components/DMXPixelMappingMatrixComponent.h b/Engine/Plugins/VirtualProduction/DMX/DMXPixelMapping/Source/DMXPixelMappingRuntime/Public/Components/DMXPixelMappingMatrixComponent.h
--- a/Engine/Plugins/VirtualProduction/DMX/DMXPixelMapping/Source/DMXPixelMappingRuntime/Public/Components/DMXPixelMappingMatrixComponent.h
+++ b/Engine/Plugins/VirtualProduction/DMX/DMXPixelMapping/Source/DMXPixelMappingRuntime/Public/Components/DMXPixelMappingMatrixComponent.h
@@ -7,8 +7,9 @@
 #include "Library/DMXEntityFixtureType.h"
 #include "DMXPixelMappingMatrixComponent.generated.h"
 
-
 class UDMXLibrary;
+class UDMXPixelMappingLayoutScript;
+
 class STextBlock;
 enum class EDMXColorMode : uint8;
 
@@ -42,6 +43,7 @@ public:
 	virtual void PostLoad() override;
 	virtual void PostInitProperties() override;
 #if WITH_EDITOR
+	virtual void PreEditChange(FProperty* PropertyAboutToChange) override;
 	virtual void PostEditChangeProperty(FPropertyChangedEvent& PropertyChangedEvent) override;
 	virtual void PostEditChangeChainProperty(FPropertyChangedChainEvent& PropertyChangedChainEvent) override;
 #endif // WITH_EDITOR
@@ -90,6 +92,12 @@ protected:
 	TArray<UDMXPixelMappingBaseComponent*> PreEditUndoMatrixCellChildren;
 #endif // WITH_EDITORONLY_DATA
 
+private:
+	/** True while the component is updating its children */
+	bool bIsUpdatingChildren = false;
+
+	/** Position before it was changed */
+	FVector2D PreEditChangePosition;
 
 public:
 	UPROPERTY()
@@ -153,4 +161,8 @@ public:
 
 	UPROPERTY()
 	EDMXPixelMappingDistribution Distribution;
+
+	/** Layout script for the children of this component (hidden in customizations and displayed in its own panel). */
+	UPROPERTY(EditAnywhere, Instanced, Category = "Layout")
+	TObjectPtr<UDMXPixelMappingLayoutScript> LayoutScript;
 };
diff --git a/Engine/Plugins/VirtualProduction/DMX/DMXPixelMapping/Source/DMXPixelMappingRuntime/Public/Components/DMXPixelMappingOutputComponent.h b/Engine/Plugins/VirtualProduction/DMX/DMXPixelMapping/Source/DMXPixelMappingRuntime/Public/Components/DMXPixelMappingOutputComponent.h
--- a/Engine/Plugins/VirtualProduction/DMX/DMXPixelMapping/Source/DMXPixelMappingRuntime/Public/Components/DMXPixelMappingOutputComponent.h
+++ b/Engine/Plugins/VirtualProduction/DMX/DMXPixelMapping/Source/DMXPixelMappingRuntime/Public/Components/DMXPixelMappingOutputComponent.h
@@ -10,28 +10,18 @@
 
 #include "DMXPixelMappingOutputComponent.generated.h"
 
+struct FDMXPixelMappingLayoutToken;
+class UDMXEntityFixturePatch;
+class UDMXPixelMappingRendererComponent;
+
 #if WITH_EDITOR
 enum class EDMXPixelMappingComponentLabelAlignment : uint8;
 class FDMXPixelMappingComponentWidget;
 #endif
-class UDMXEntityFixturePatch;
-class UDMXPixelMappingRendererComponent;
 
 class SBox;
 
 
-/** Definition for Default colors */
-struct FDMXOutputComponentColors
-{
-	// Note, for the default color, use the editor color
-
-	/** The color used when no the component is selected */
-	virtual const FLinearColor& GetSelectedColor() { return SelectedColor; }
-
-private:
-	static const FLinearColor SelectedColor;
-};
-
 /** Enum that defines the quality of how pixels are rendered */
 UENUM()
 enum class EDMXPixelBlendingQuality : uint8
@@ -87,7 +77,10 @@ public:
 	virtual const FText GetPaletteCategory();
 
 	/** Rebuild widget for designer view */
+	UE_DEPRECATED(5.1, "Pixel Mapping Components no longer hold their own widget, in an effort to separate Views from Data.")
+	PRAGMA_DISABLE_DEPRECATION_WARNINGS // FDMXPixelMappingComponentWidget is deprecated
 	virtual TSharedRef<FDMXPixelMappingComponentWidget> BuildSlot(TSharedRef<SConstraintCanvas> InCanvas);
+	PRAGMA_ENABLE_DEPRECATION_WARNINGS
 
 	/** Whether component should be visible */
 	virtual bool IsVisible() const;
@@ -123,13 +116,13 @@ public:
 	/** Sets the position */
 	virtual void SetPosition(const FVector2D& NewPosition);
 
-	/** Returns the position */
+	/** Returns the position. Note this is the position before a layout has been applied, use GetLayoutedPosition() to get the final position.. */
 	FVector2D GetPosition() const { return FVector2D(PositionX, PositionY); }
 
 	/** Sets the size */
 	virtual void SetSize(const FVector2D& NewSize);
 
-	/** Get the size of the component */
+	/** Get the size. Note this is the position before a layout has been applied, use GetLayoutedSize() to get the final position.. */
 	FVector2D GetSize() const { return FVector2D(SizeX, SizeY); }
 
 	/** Helper that returns render component if available */
@@ -144,7 +137,10 @@ public:
 public:
 #if WITH_EDITOR
 	/** Returns the component widget */
-	FORCEINLINE const TSharedPtr<FDMXPixelMappingComponentWidget> GetComponentWidget() { return ComponentWidget; }
+	UE_DEPRECATED(5.1, "Pixel Mapping Components no longer hold their own widget, in an effort to separate Views from Data.")
+	PRAGMA_DISABLE_DEPRECATION_WARNINGS // FDMXPixelMappingComponentWidget is deprecated
+	FORCEINLINE const TSharedPtr<FDMXPixelMappingComponentWidget> GetComponentWidget() { return ComponentWidget_DEPRECATED; }
+	PRAGMA_ENABLE_DEPRECATION_WARNINGS
 
 protected:
 	/** Udpates the component widget. If bWithChildrenRecursive, updates child Components' Component Widget recursively */
@@ -168,8 +164,10 @@ protected:
 	/** Children available PreEditUndo, useful to hide all removed ones in post edit undo */
 	TArray<UDMXPixelMappingBaseComponent*> PreEditUndoChildren;
 
+	PRAGMA_DISABLE_DEPRECATION_WARNINGS
 	/** The widget shown for this component */
-	TSharedPtr<FDMXPixelMappingComponentWidget> ComponentWidget;
+	TSharedPtr<FDMXPixelMappingComponentWidget> ComponentWidget_DEPRECATED;
+	PRAGMA_ENABLE_DEPRECATION_WARNINGS
 #endif // WITH_EDITORONLY_DATA
 
 
@@ -178,27 +176,30 @@ protected:
 	----------------------------------*/
 public:
 #if WITH_EDITOR
-	/** Returns the bIsLockInDesigner property name */
+	// Property Name getters
 	FORCEINLINE static FName GetLockInDesignerPropertyName() { return GET_MEMBER_NAME_CHECKED(UDMXPixelMappingOutputComponent, bLockInDesigner); }
-
-	/** Returns the bIsVisibleInDesigner property name */
 	FORCEINLINE static FName GetVisibleInDesignerPropertyName() { return GET_MEMBER_NAME_CHECKED(UDMXPixelMappingOutputComponent, bVisibleInDesigner); }
+	FORCEINLINE static FName GetPositionXPropertyName() { return GET_MEMBER_NAME_CHECKED(UDMXPixelMappingOutputComponent, PositionX); }
+	FORCEINLINE static FName GetPositionYPropertyName() { return GET_MEMBER_NAME_CHECKED(UDMXPixelMappingOutputComponent, PositionY); }
+	FORCEINLINE static FName GetSizeXPropertyName() { return GET_MEMBER_NAME_CHECKED(UDMXPixelMappingOutputComponent, SizeX); }
+	FORCEINLINE static FName GetSizeYPropertyName() { return GET_MEMBER_NAME_CHECKED(UDMXPixelMappingOutputComponent, SizeY); }
 #endif
 
 	/** The quality level to use when averaging colors during downsampling. */
 	UPROPERTY(EditAnywhere, BlueprintReadOnly, Category = "Pixel Settings")
 	EDMXPixelBlendingQuality CellBlendingQuality;
 
-	UPROPERTY(EditAnywhere, BlueprintReadOnly, Category = "Common Settings", Meta = (EditCondition = "!bLockInDesigner"))
+private:
+	UPROPERTY(EditAnywhere, BlueprintReadOnly, Category = "Common Settings", Meta = (EditCondition = "!bLockInDesigner", AllowPrivateAccess = true))
 	float PositionX;
 
-	UPROPERTY(EditAnywhere, BlueprintReadOnly, Category = "Common Settings", Meta = (EditCondition = "!bLockInDesigner"))
+	UPROPERTY(EditAnywhere, BlueprintReadOnly, Category = "Common Settings", Meta = (EditCondition = "!bLockInDesigner", AllowPrivateAccess = true))
 	float PositionY;
 
-	UPROPERTY(EditAnywhere, BlueprintReadOnly, Category = "Common Settings", Meta = (ClampMin = 0, UIMin = 0, EditCondition = "!bLockInDesigner"))
+	UPROPERTY(EditAnywhere, BlueprintReadOnly, Category = "Common Settings", Meta = (ClampMin = 1, UIMin = 1, EditCondition = "!bLockInDesigner", AllowPrivateAccess = true))
 	float SizeX;
 
-	UPROPERTY(EditAnywhere, BlueprintReadOnly, Category = "Common Settings", Meta = (ClampMin = 0, UIMin = 0, EditCondition = "!bLockInDesigner"))
+	UPROPERTY(EditAnywhere, BlueprintReadOnly, Category = "Common Settings", Meta = (ClampMin = 1, UIMin = 1, EditCondition = "!bLockInDesigner", AllowPrivateAccess = true))
 	float SizeY;
 
 protected:
diff --git a/Engine/Plugins/VirtualProduction/DMX/DMXPixelMapping/Source/DMXPixelMappingRuntime/Public/Components/DMXPixelMappingRendererComponent.h b/Engine/Plugins/VirtualProduction/DMX/DMXPixelMapping/Source/DMXPixelMappingRuntime/Public/Components/DMXPixelMappingRendererComponent.h
--- a/Engine/Plugins/VirtualProduction/DMX/DMXPixelMapping/Source/DMXPixelMappingRuntime/Public/Components/DMXPixelMappingRendererComponent.h
+++ b/Engine/Plugins/VirtualProduction/DMX/DMXPixelMapping/Source/DMXPixelMappingRuntime/Public/Components/DMXPixelMappingRendererComponent.h
@@ -11,15 +11,16 @@
 
 #include "DMXPixelMappingRendererComponent.generated.h"
 
+enum class EDMXPixelMappingRendererType : uint8;
+class UDMXPixelMappingLayoutScript;
 
+enum class EMapChangeType : uint8;
 class UMaterialInterface;
 class UTexture;
 class UUserWidget;
 class UTextureRenderTarget2D;
 class UWorld;
 
-enum class EMapChangeType : uint8;
-enum class EDMXPixelMappingRendererType : uint8;
 
 /**
  * Component for rendering input texture
@@ -55,10 +56,6 @@ public:
 	virtual void RenderAndSendDMX() final;
 	//~ End UDMXPixelMappingBaseComponent implementation
 
-	// ~Begin UDMXPixelMappingOutputComponent interface
-	virtual void SetSize(const FVector2D& NewSize) override;
-	// ~End UDMXPixelMappingOutputComponent interface
-
 #if WITH_EDITOR
 	/** Render all downsample pixel for editor preview texture */
 	void RenderEditorPreviewTexture();
@@ -88,6 +85,7 @@ public:
 	/**
 	 * Take of container widget which is holds widget for all child components.
 	 */
+	UE_DEPRECATED(5.1, "Pixel Mapping Components no longer hold their own widget, in an effort to separate Views from Data.")
 	TSharedRef<SWidget> TakeWidget();
 #endif // WITH_EDITOR
 
@@ -184,6 +182,10 @@ public:
 	UPROPERTY(EditAnywhere, BlueprintReadOnly, Category = "Render Settings", meta = (ClampMin = "0", ClampMax = "1", UIMin = "0", UIMax = "1"))
 	float Brightness;
 
+	/** Layout script for the children of this component (hidden in customizations and displayed in its own panel). */
+	UPROPERTY(EditAnywhere, Instanced, Category = "Layout")
+	TObjectPtr<UDMXPixelMappingLayoutScript> LayoutScript;
+
 	/** Check if a Component can be moved under another one (used for copy/move/duplicate) */
 	virtual bool CanBeMovedTo(const UDMXPixelMappingBaseComponent* Component) const override;
 
@@ -242,6 +244,7 @@ private:
 
 	/** Initial texture color */
 	static const FLinearColor ClearTextureColor;
+
 public:
 	/** Max downsample target size */
 	static const FIntPoint MaxDownsampleBufferTargetSize;
diff --git a/Engine/Plugins/VirtualProduction/DMX/DMXPixelMapping/Source/DMXPixelMappingRuntime/Public/Components/DMXPixelMappingScreenComponent.h b/Engine/Plugins/VirtualProduction/DMX/DMXPixelMapping/Source/DMXPixelMappingRuntime/Public/Components/DMXPixelMappingScreenComponent.h
--- a/Engine/Plugins/VirtualProduction/DMX/DMXPixelMapping/Source/DMXPixelMappingRuntime/Public/Components/DMXPixelMappingScreenComponent.h
+++ b/Engine/Plugins/VirtualProduction/DMX/DMXPixelMapping/Source/DMXPixelMappingRuntime/Public/Components/DMXPixelMappingScreenComponent.h
@@ -49,7 +49,9 @@ public:
 #if WITH_EDITOR
 	virtual const FText GetPaletteCategory() override;
 	virtual bool IsExposedToTemplate() { return true; }
+	PRAGMA_DISABLE_DEPRECATION_WARNINGS // FDMXPixelMappingComponentWidget is deprecated
 	virtual TSharedRef<FDMXPixelMappingComponentWidget> BuildSlot(TSharedRef<SConstraintCanvas> InCanvas) override;
+	PRAGMA_ENABLE_DEPRECATION_WARNINGS
 #endif // WITH_EDITOR
 	virtual void SetPosition(const FVector2D& NewPosition) override;
 	virtual void SetSize(const FVector2D& NewSize) override;
@@ -100,10 +102,10 @@ private:
 	TSet<FDMXOutputPortSharedRef> OutputPorts;
 
 public:
-	UPROPERTY(EditAnywhere, BlueprintReadOnly, Category = "Mapping Settings", meta = (DisplayName = "X Cells", ClampMin = "1", ClampMax = "1000", UIMin = "1", UIMax = "1000", DisplayPriority = "1"))
+	UPROPERTY(EditAnywhere, BlueprintReadOnly, Category = "Mapping Settings", meta = (DisplayName = "X Cells", ClampMin = "1", ClampMax = "1000", UIMin = "1", UIMax = "100", DisplayPriority = "1"))
 	int32 NumXCells;
 
-	UPROPERTY(EditAnywhere, BlueprintReadOnly, Category = "Mapping Settings", meta = (DisplayName = "Y Cells", ClampMin = "1", ClampMax = "1000", UIMin = "1", UIMax = "1000", DisplayPriority = "1"))
+	UPROPERTY(EditAnywhere, BlueprintReadOnly, Category = "Mapping Settings", meta = (DisplayName = "Y Cells", ClampMin = "1", ClampMax = "1000", UIMin = "1", UIMax = "100", DisplayPriority = "1"))
 	int32 NumYCells;
 
 	UPROPERTY(meta = (DeprecatedProperty, DeprecationMessage = "Use Output Ports instead."))
