DMX - Remove unacceptable words from DMXPixelMappingRendererComponent

diff --git a/Engine/Plugins/VirtualProduction/DMX/DMXPixelMapping/Source/DMXPixelMappingRuntime/Public/Components/DMXPixelMappingRendererComponent.h b/Engine/Plugins/VirtualProduction/DMX/DMXPixelMapping/Source/DMXPixelMappingRuntime/Public/Components/DMXPixelMappingRendererComponent.h
--- a/Engine/Plugins/VirtualProduction/DMX/DMXPixelMapping/Source/DMXPixelMappingRuntime/Public/Components/DMXPixelMappingRendererComponent.h
+++ b/Engine/Plugins/VirtualProduction/DMX/DMXPixelMapping/Source/DMXPixelMappingRuntime/Public/Components/DMXPixelMappingRendererComponent.h
@@ -178,7 +178,7 @@ public:
 	UPROPERTY(EditAnywhere, BlueprintReadOnly, Category = "Render Settings")
 	TSubclassOf<UUserWidget> InputWidget;
 	
-	/** Master brightness of the renderer */
+	/** The brightness of the renderer */
 	UPROPERTY(EditAnywhere, BlueprintReadOnly, Category = "Render Settings", meta = (ClampMin = "0", ClampMax = "1", UIMin = "0", UIMax = "1"))
 	float Brightness;
 
