DMX - Undo now works as expected when deleting Pixel Mapping Renderer Components

diff --git a/Engine/Plugins/VirtualProduction/DMX/DMXPixelMapping/Source/DMXPixelMappingEditor/Private/Toolkits/DMXPixelMappingToolkit.cpp b/Engine/Plugins/VirtualProduction/DMX/DMXPixelMapping/Source/DMXPixelMappingEditor/Private/Toolkits/DMXPixelMappingToolkit.cpp
--- a/Engine/Plugins/VirtualProduction/DMX/DMXPixelMapping/Source/DMXPixelMappingEditor/Private/Toolkits/DMXPixelMappingToolkit.cpp
+++ b/Engine/Plugins/VirtualProduction/DMX/DMXPixelMapping/Source/DMXPixelMappingEditor/Private/Toolkits/DMXPixelMappingToolkit.cpp
@@ -520,40 +520,16 @@ void FDMXPixelMappingToolkit::DeleteSelectedComponents()
 	{
 		TGuardValue<bool>(bRemovingComponents, true);
 
-		// Count how many Descendants need to be removed 
-		const int32 NumComponentsAndItsDescendants = [this]()
-		{
-			int32 DescendantsCounter = 0;
-			for (const FDMXPixelMappingComponentReference& ComponentReference : SelectedComponents)
-			{
-				if (UDMXPixelMappingBaseComponent* Component = ComponentReference.GetComponent())
-				{
-					DescendantsCounter++;
-
-					constexpr bool bCountRecursive = true;
-					Component->ForEachChild([&DescendantsCounter](UDMXPixelMappingBaseComponent* ChildComponent)
-						{
-							DescendantsCounter++;
-						}, bCountRecursive);
-				}
-			}
-			return DescendantsCounter;
-		}();
-
 		TSet<FDMXPixelMappingComponentReference> ParentComponentReferences;
 		for (const FDMXPixelMappingComponentReference& ComponentReference : SelectedComponents)
 		{
 			if (UDMXPixelMappingBaseComponent* Component = ComponentReference.GetComponent())
 			{
-				if (UDMXPixelMappingRendererComponent* RendererComponent = Cast<UDMXPixelMappingRendererComponent>(Component))
-				{
-					SetActiveRenderComponent(nullptr);
-				}
-
-				if (UDMXPixelMappingOutputComponent* OutputComponent = Cast<UDMXPixelMappingOutputComponent>(Component))
-				{
-					ActiveOutputComponents.Remove(OutputComponent);
-				}
+				constexpr bool bModifyChildrenRecursively = true;
+				Component->ForEachChild([](UDMXPixelMappingBaseComponent* ChildComponent)
+					{
+						ChildComponent->Modify();
+					}, bModifyChildrenRecursively);
 
 				if (UDMXPixelMappingBaseComponent* ParentComponent = Component->GetParent())
 				{
@@ -563,7 +539,17 @@ void FDMXPixelMappingToolkit::DeleteSelectedComponents()
 					Component->Modify();
 					ParentComponent->Modify();
 
-					ParentComponent->RemoveChild(ComponentReference.GetComponent());
+					ParentComponent->RemoveChild(Component);
+				}
+			
+				if (UDMXPixelMappingOutputComponent* OutputComponent = Cast<UDMXPixelMappingOutputComponent>(Component))
+				{
+					ActiveOutputComponents.Remove(OutputComponent);
+				}
+
+				if (UDMXPixelMappingRendererComponent* RendererComponent = Cast<UDMXPixelMappingRendererComponent>(Component))
+				{
+					SetActiveRenderComponent(nullptr);
 				}
 			}
 		}
diff --git a/Engine/Plugins/VirtualProduction/DMX/DMXPixelMapping/Source/DMXPixelMappingEditor/Private/Views/SDMXPixelMappingDesignerView.cpp b/Engine/Plugins/VirtualProduction/DMX/DMXPixelMapping/Source/DMXPixelMappingEditor/Private/Views/SDMXPixelMappingDesignerView.cpp
--- a/Engine/Plugins/VirtualProduction/DMX/DMXPixelMapping/Source/DMXPixelMappingEditor/Private/Views/SDMXPixelMappingDesignerView.cpp
+++ b/Engine/Plugins/VirtualProduction/DMX/DMXPixelMapping/Source/DMXPixelMappingEditor/Private/Views/SDMXPixelMappingDesignerView.cpp
@@ -28,6 +28,7 @@
 #include "Templates/DMXPixelMappingComponentTemplate.h"
 #include "Toolkits/DMXPixelMappingToolkit.h"
 
+#include "ScopedTransaction.h"
 #include "Widgets/SCanvas.h"
 #include "Widgets/Text/STextBlock.h"
 #include "Widgets/Layout/SBox.h"
@@ -409,6 +410,8 @@ FReply SDMXPixelMappingDesignerView::OnKeyDown(const FGeometry& MyGeometry, cons
 	{
 		if (const TSharedPtr<FDMXPixelMappingToolkit> ToolkitPtr = ToolkitWeakPtr.Pin())
 		{
+			const FScopedTransaction Transaction(FText::Format(LOCTEXT("DMXPixelMapping.RemoveComponents", "Remove {0}|plural(one=Component, other=Components)"), ToolkitPtr->GetSelectedComponents().Num()));
+
 			ToolkitPtr->DeleteSelectedComponents();
 
 			return FReply::Handled();
diff --git a/Engine/Plugins/VirtualProduction/DMX/DMXPixelMapping/Source/DMXPixelMappingEditor/Private/Views/SDMXPixelMappingHierarchyView.cpp b/Engine/Plugins/VirtualProduction/DMX/DMXPixelMapping/Source/DMXPixelMappingEditor/Private/Views/SDMXPixelMappingHierarchyView.cpp
--- a/Engine/Plugins/VirtualProduction/DMX/DMXPixelMapping/Source/DMXPixelMappingEditor/Private/Views/SDMXPixelMappingHierarchyView.cpp
+++ b/Engine/Plugins/VirtualProduction/DMX/DMXPixelMapping/Source/DMXPixelMappingEditor/Private/Views/SDMXPixelMappingHierarchyView.cpp
@@ -637,19 +637,23 @@ void SDMXPixelMappingHierarchyView::BeginDelete()
 
 	if (TSharedPtr<FDMXPixelMappingToolkit> ToolkitPtr = Toolkit.Pin())
 	{
-		TSet<FDMXPixelMappingComponentReference> SelectedComponents;
-
 		for (FDMXPixelMappingHierarchyItemWidgetModelPtr SelectedItem : SelectedItems)
 		{
 			UDMXPixelMappingBaseComponent* SelectedComponent = SelectedItem->GetReference().GetComponent();
 			if (SelectedComponent)
 			{
+				constexpr bool bModifyChildrenRecursively = true;
+				SelectedComponent->ForEachChild([](UDMXPixelMappingBaseComponent* ChildComponent)
+					{
+						ChildComponent->Modify();
+					}, bModifyChildrenRecursively);
+
 				UDMXPixelMappingBaseComponent* ParentComponent = SelectedComponent->GetParent();
 				if (ParentComponent)
 				{
 					ParentComponent->Modify();
 					SelectedComponent->Modify();
-					
+
 					ParentComponent->RemoveChild(SelectedComponent);
 				}
 			}
