DMX - Fix a crash when Pixel Mapping Modulator outputs an array that is smaller than the input array

diff --git a/Engine/Plugins/VirtualProduction/DMX/DMXPixelMapping/Source/DMXPixelMappingRuntime/Private/Components/DMXPixelMappingMatrixComponent.cpp b/Engine/Plugins/VirtualProduction/DMX/DMXPixelMapping/Source/DMXPixelMappingRuntime/Private/Components/DMXPixelMappingMatrixComponent.cpp
--- a/Engine/Plugins/VirtualProduction/DMX/DMXPixelMapping/Source/DMXPixelMappingRuntime/Private/Components/DMXPixelMappingMatrixComponent.cpp
+++ b/Engine/Plugins/VirtualProduction/DMX/DMXPixelMapping/Source/DMXPixelMappingRuntime/Private/Components/DMXPixelMappingMatrixComponent.cpp
@@ -275,7 +275,10 @@ void UDMXPixelMappingMatrixComponent::SendDMX()
 			if (UDMXPixelMappingMatrixCellComponent* CellComponent = Cast<UDMXPixelMappingMatrixCellComponent>(Children[IndexChild]))
 			{
 				// Relies on the order of childs and AttributeToValueMapArray didn't change during the lifetime of this function!
-				check(AttributeToValueMapArray.IsValidIndex(IndexChild));
+				if (!AttributeToValueMapArray.IsValidIndex(IndexChild))
+				{
+					break;
+				}
 
 				TMap<int32, uint8> ChannelToValueMap;
 				for (const TTuple<FDMXAttributeName, float>& AttributeValuePair : AttributeToValueMapArray[IndexChild].Map)
