DMX - Pixel Mapping zoom when mouse is over output mapping now works as expected work

diff --git a/Engine/Plugins/VirtualProduction/DMX/DMXPixelMapping/Source/DMXPixelMappingEditor/Private/Widgets/SDMXPixelMappingPreviewViewport.cpp b/Engine/Plugins/VirtualProduction/DMX/DMXPixelMapping/Source/DMXPixelMappingEditor/Private/Widgets/SDMXPixelMappingPreviewViewport.cpp
--- a/Engine/Plugins/VirtualProduction/DMX/DMXPixelMapping/Source/DMXPixelMappingEditor/Private/Widgets/SDMXPixelMappingPreviewViewport.cpp
+++ b/Engine/Plugins/VirtualProduction/DMX/DMXPixelMapping/Source/DMXPixelMappingEditor/Private/Widgets/SDMXPixelMappingPreviewViewport.cpp
@@ -21,6 +21,7 @@ void SDMXPixelMappingPreviewViewport::Construct(const FArguments& InArgs, const
 	ChildSlot
 		[
 			SNew(SBox)
+			.Visibility(EVisibility::HitTestInvisible)
 			.WidthOverride(this, &SDMXPixelMappingPreviewViewport::GetPreviewAreaWidth)
 			.HeightOverride(this, &SDMXPixelMappingPreviewViewport::GetPreviewAreaHeight)
 			[
