DMX - Fix DMX custom attributes do not populate in Pixel Mapping

diff --git a/Engine/Plugins/VirtualProduction/DMX/DMXEngine/Source/DMXEditor/Private/Customizations/DMXEntityReferenceCustomization.cpp b/Engine/Plugins/VirtualProduction/DMX/DMXEngine/Source/DMXEditor/Private/Customizations/DMXEntityReferenceCustomization.cpp
--- a/Engine/Plugins/VirtualProduction/DMX/DMXEngine/Source/DMXEditor/Private/Customizations/DMXEntityReferenceCustomization.cpp
+++ b/Engine/Plugins/VirtualProduction/DMX/DMXEngine/Source/DMXEditor/Private/Customizations/DMXEntityReferenceCustomization.cpp
@@ -132,7 +132,9 @@ FText FDMXEntityReferenceCustomization::GetPickerPropertyLabel() const
 
 bool FDMXEntityReferenceCustomization::GetPickerEnabled() const
 {
-	return GetEntityType() != nullptr;
+	return  
+		StructHandle->IsEditable() && 
+		GetEntityType() != nullptr;
 }
 
 TWeakObjectPtr<UDMXEntity> FDMXEntityReferenceCustomization::GetCurrentEntity() const
diff --git a/Engine/Plugins/VirtualProduction/DMX/DMXPixelMapping/Source/DMXPixelMappingEditor/DMXPixelMappingEditor.Build.cs b/Engine/Plugins/VirtualProduction/DMX/DMXPixelMapping/Source/DMXPixelMappingEditor/DMXPixelMappingEditor.Build.cs
--- a/Engine/Plugins/VirtualProduction/DMX/DMXPixelMapping/Source/DMXPixelMappingEditor/DMXPixelMappingEditor.Build.cs
+++ b/Engine/Plugins/VirtualProduction/DMX/DMXPixelMapping/Source/DMXPixelMappingEditor/DMXPixelMappingEditor.Build.cs
@@ -17,6 +17,7 @@ public class DMXPixelMappingEditor : ModuleRules
 			"DMXPixelMappingRuntime",
 			"DMXPixelMappingRenderer",
 			"DMXPixelMappingBlueprintGraph",
+			"DMXProtocol",
 			"DMXRuntime",
 			"EditorStyle",
 			"Engine",
diff --git a/Engine/Plugins/VirtualProduction/DMX/DMXPixelMapping/Source/DMXPixelMappingEditor/Private/Customizations/DMXPixelMappingDetailCustomization_FixtureGroupItem.cpp b/Engine/Plugins/VirtualProduction/DMX/DMXPixelMapping/Source/DMXPixelMappingEditor/Private/Customizations/DMXPixelMappingDetailCustomization_FixtureGroupItem.cpp
--- a/Engine/Plugins/VirtualProduction/DMX/DMXPixelMapping/Source/DMXPixelMappingEditor/Private/Customizations/DMXPixelMappingDetailCustomization_FixtureGroupItem.cpp
+++ b/Engine/Plugins/VirtualProduction/DMX/DMXPixelMapping/Source/DMXPixelMappingEditor/Private/Customizations/DMXPixelMappingDetailCustomization_FixtureGroupItem.cpp
@@ -4,8 +4,11 @@
 
 #include "DMXPixelMappingTypes.h"
 #include "Components/DMXPixelMappingFixtureGroupItemComponent.h"
-#include "Toolkits/DMXPixelMappingToolkit.h"
+#include "Library/DMXEntityFixturePatch.h"
+#include "Library/DMXEntityReference.h"
 #include "Modulators/DMXModulator.h"
+#include "Toolkits/DMXPixelMappingToolkit.h"
+#include "Widgets/SDMXPixelMappingAttributeNamesComboBox.h"
 
 #include "DetailLayoutBuilder.h"
 #include "DetailCategoryBuilder.h"
@@ -16,8 +19,72 @@
 #include "Layout/Visibility.h"
 #include "Modules/ModuleManager.h"
 
+
+
 #define LOCTEXT_NAMESPACE "FixtureGroupItem"
 
+bool FDMXPixelMappingDetailCustomization_FixtureGroupItem::FFunctionAttribute::HasMultipleAttributeValues() const
+{
+	TArray<const void*> RawData;
+	Handle->AccessRawData(RawData);
+
+	TSet<FName> AttributeNames;
+	for (const void* RawPtr : RawData)
+	{
+		if (RawPtr)
+		{
+			// The types we use with this customization must have a cast constructor to FName
+			AttributeNames.Add(reinterpret_cast<const FDMXAttributeName*>(RawPtr)->Name);
+		}
+	}
+
+	return AttributeNames.Num() > 1;
+}
+
+FName FDMXPixelMappingDetailCustomization_FixtureGroupItem::FFunctionAttribute::GetAttributeValue() const
+{
+	if (!ensureMsgf(!HasMultipleAttributeValues(), TEXT("Cannot get attribute value from handle when handle has mutliple values")))
+	{
+		return NAME_None;
+	}
+
+	TArray<const void*> RawData;
+	Handle->AccessRawData(RawData);
+
+	for (const void* RawPtr : RawData)
+	{
+		if (RawPtr)
+		{
+			// The types we use with this customization must have a cast constructor to FName
+			return reinterpret_cast<const FDMXAttributeName*>(RawPtr)->Name;
+		}
+	}
+
+	return NAME_None;
+}
+
+void FDMXPixelMappingDetailCustomization_FixtureGroupItem::FFunctionAttribute::SetAttributeValue(const FName& NewValue)
+{
+	FStructProperty* StructProperty = CastFieldChecked<FStructProperty>(Handle->GetProperty());
+
+	TArray<void*> RawData;
+	Handle->AccessRawData(RawData);
+
+	for (void* SingleRawData : RawData)
+	{
+		FDMXAttributeName* PreviousValue = reinterpret_cast<FDMXAttributeName*>(SingleRawData);
+		FDMXAttributeName NewAttributeName;
+		NewAttributeName.SetFromName(NewValue);
+
+		// Export new value to text format that can be imported later
+		FString TextValue;
+		StructProperty->Struct->ExportText(TextValue, &NewAttributeName, PreviousValue, nullptr, EPropertyPortFlags::PPF_None, nullptr);
+
+		// Set values on edited property handle from exported text
+		ensure(Handle->SetValueFromFormattedString(TextValue, EPropertyValueSetFlags::DefaultFlags) == FPropertyAccess::Result::Success);
+	}
+}
+
 void FDMXPixelMappingDetailCustomization_FixtureGroupItem::CustomizeDetails(IDetailLayoutBuilder& InDetailLayout)
 {
 	DetailLayout = &InDetailLayout;
@@ -28,7 +95,6 @@ void FDMXPixelMappingDetailCustomization_FixtureGroupItem::CustomizeDetails(IDet
 	DetailLayout->GetObjectsBeingCustomized(Objects);
 
 	FixtureGroupItemComponents.Empty();
-	
 	for (TWeakObjectPtr<UObject> SelectedObject : Objects)
 	{
 		FixtureGroupItemComponents.Add(Cast<UDMXPixelMappingFixtureGroupItemComponent>(SelectedObject));
@@ -83,7 +149,9 @@ void FDMXPixelMappingDetailCustomization_FixtureGroupItem::CustomizeDetails(IDet
 		.Visibility(TAttribute<EVisibility>(this, &FDMXPixelMappingDetailCustomization_FixtureGroupItem::GetRGBAttributesVisibility))
 		.NameContent()
 		[
-			SNew(STextBlock).Text(LOCTEXT("ColorSample", "Color Sample"))
+			SNew(STextBlock)
+			.Text(LOCTEXT("ColorSample", "Color Sample"))
+			.Font(FAppStyle::GetFontStyle(TEXT("PropertyWindow.NormalFont")))
 		]
 		.ValueContent()
 		[
@@ -92,23 +160,13 @@ void FDMXPixelMappingDetailCustomization_FixtureGroupItem::CustomizeDetails(IDet
 			.OnGenerateRow(this, &FDMXPixelMappingDetailCustomization_FixtureGroupItem::GenerateExposeAndInvertRow)
 		];
 
-	// Update RGB attributes
-	for (TSharedPtr<FFunctionAttribute>& Attribute : RGBAttributes)
-	{
-		DetailLayout->HideProperty(Attribute->ExposeHandle);
-		DetailLayout->HideProperty(Attribute->InvertHandle);
-
-		OutputSettingsCategory
-			.AddProperty(Attribute->Handle)
-			.Visibility(TAttribute<EVisibility>::Create(TAttribute<EVisibility>::FGetter::CreateSP(this, &FDMXPixelMappingDetailCustomization_FixtureGroupItem::GetRGBAttributeRowVisibilty, Attribute.Get())));
-	}
-
 	// Generate all Monochrome Expose and Invert rows
 	OutputSettingsCategory.AddCustomRow(FText::GetEmpty())
 		.Visibility(TAttribute<EVisibility>(this, &FDMXPixelMappingDetailCustomization_FixtureGroupItem::GetMonochromeAttributesVisibility))
 		.NameContent()
 		[
 			SNew(STextBlock).Text(LOCTEXT("ColorSample", "Color Sample"))
+			.Font(FAppStyle::GetFontStyle(TEXT("PropertyWindow.NormalFont")))
 		]
 		.ValueContent()
 		[
@@ -117,18 +175,8 @@ void FDMXPixelMappingDetailCustomization_FixtureGroupItem::CustomizeDetails(IDet
 			.OnGenerateRow(this, &FDMXPixelMappingDetailCustomization_FixtureGroupItem::GenerateExposeAndInvertRow)
 		];
 
-	// Update Monochrome attributes
-	for (TSharedPtr<FFunctionAttribute>& Attribute : MonochromeAttributes)
-	{
-		DetailLayout->HideProperty(Attribute->ExposeHandle);
-		DetailLayout->HideProperty(Attribute->InvertHandle);
-
-		OutputSettingsCategory
-			.AddProperty(Attribute->Handle)
-			.Visibility(TAttribute<EVisibility>::Create(TAttribute<EVisibility>::FGetter::CreateSP(this, &FDMXPixelMappingDetailCustomization_FixtureGroupItem::GetMonochromeRowVisibilty, Attribute.Get())));
-	}
-
 	CreateModulatorDetails(InDetailLayout);
+	CreateAttributeDetails(InDetailLayout);
 }
 
 bool FDMXPixelMappingDetailCustomization_FixtureGroupItem::CheckComponentsDMXColorMode(const EDMXColorMode DMXColorMode) const
@@ -326,7 +374,7 @@ void FDMXPixelMappingDetailCustomization_FixtureGroupItem::CreateModulatorDetail
 						[
 							SNew(STextBlock)
 							.Text(LOCTEXT("ModulatorMultipleValues", "Multiple Values"))
-						.Font(IDetailLayoutBuilder::GetDetailFont())
+							.Font(IDetailLayoutBuilder::GetDetailFont())
 						];
 
 					break;
@@ -336,6 +384,120 @@ void FDMXPixelMappingDetailCustomization_FixtureGroupItem::CreateModulatorDetail
 	}
 }
 
+void FDMXPixelMappingDetailCustomization_FixtureGroupItem::CreateAttributeDetails(IDetailLayoutBuilder& InDetailLayout)
+{
+	TArray<TWeakObjectPtr<UObject>> Objects;
+	DetailLayout->GetObjectsBeingCustomized(Objects);
+
+	TArray<FName> FixtureGroupItemComponentsAttributes;
+
+	TArray<UDMXPixelMappingFixtureGroupItemComponent*> GroupItemComponents;
+	for (TWeakObjectPtr<UObject> SelectedObject : Objects)
+	{
+		UDMXPixelMappingFixtureGroupItemComponent* FixtureGroupItemComponent = Cast<UDMXPixelMappingFixtureGroupItemComponent>(SelectedObject);
+		if (FixtureGroupItemComponent)
+		{
+			GroupItemComponents.Add(FixtureGroupItemComponent);
+		}
+	}
+
+	// Gather attribute names present in Patches of all selected components
+	for (UDMXPixelMappingFixtureGroupItemComponent* GroupItemComponent : GroupItemComponents)
+	{
+		const UDMXEntityFixturePatch* FixturePatch = GroupItemComponent->FixturePatchRef.GetFixturePatch();
+		if (!FixturePatch)
+		{
+			continue;
+		}
+
+		const TArray<FDMXAttributeName> AttributeNames = FixturePatch->GetAllAttributesInActiveMode();
+		for (const FDMXAttributeName& AttributeName : AttributeNames)
+		{
+			FixtureGroupItemComponentsAttributes.AddUnique(AttributeName.Name);
+		}
+	}
+
+	for (UDMXPixelMappingFixtureGroupItemComponent* GroupItemComponent : GroupItemComponents)
+	{
+		FixtureGroupItemComponentsAttributes.RemoveAll([GroupItemComponent](const FName& AttributeName)
+			{
+				const UDMXEntityFixturePatch* FixturePatch = GroupItemComponent->FixturePatchRef.GetFixturePatch();
+				if (!FixturePatch)
+				{
+					return true;
+				}
+				const TArray<FDMXAttributeName> AttributeNamesOfPatch = FixturePatch->GetAllAttributesInActiveMode();
+				return !AttributeNamesOfPatch.Contains(AttributeName);
+			});
+	}
+
+	// Lambda to add a single attribute property row
+	auto CreateAttributePropertyRowLambda = [&InDetailLayout, &FixtureGroupItemComponentsAttributes](const TSharedRef<FFunctionAttribute>& Attribute, const TAttribute<EVisibility>& VisibilityAttribute)
+		{
+			InDetailLayout.HideProperty(Attribute->ExposeHandle);
+			InDetailLayout.HideProperty(Attribute->InvertHandle);
+			InDetailLayout.HideProperty(Attribute->Handle);
+
+			const TSharedRef< SDMXPixelMappingAttributeNamesComboBox> AttributeNameComboBox =
+				SNew(SDMXPixelMappingAttributeNamesComboBox, FixtureGroupItemComponentsAttributes)
+				.OnSelectionChanged_Lambda([Attribute](const FName& NewValue)
+					{
+						Attribute->SetAttributeValue(NewValue);
+					});
+
+			if (Attribute->HasMultipleAttributeValues())
+			{
+				AttributeNameComboBox->SetHasMultipleValues();
+			}
+			else if (FixtureGroupItemComponentsAttributes.Num() > 0)
+			{
+				FName InitialSelection = Attribute->GetAttributeValue();
+
+				// Set a valid attribute, or name none if none is available
+				if (!FixtureGroupItemComponentsAttributes.Contains(InitialSelection))
+				{
+					if (FixtureGroupItemComponentsAttributes.IsEmpty())
+					{
+						Attribute->SetAttributeValue(NAME_None);
+					}
+					else
+					{
+						Attribute->SetAttributeValue(FixtureGroupItemComponentsAttributes[0]);
+						InitialSelection = FixtureGroupItemComponentsAttributes[0];
+					}
+				}
+				AttributeNameComboBox->SetSelection(InitialSelection);
+			}
+
+			InDetailLayout.EditCategory("Output Settings", FText::GetEmpty(), ECategoryPriority::Important)
+				.AddCustomRow(FText::GetEmpty())
+				.Visibility(VisibilityAttribute)
+				.NameContent()
+				[
+					SNew(STextBlock)
+					.Text(Attribute->Handle->GetPropertyDisplayName())
+					.Font(FAppStyle::GetFontStyle(TEXT("PropertyWindow.NormalFont")))
+				]
+				.ValueContent()
+				[
+					AttributeNameComboBox
+				];
+		};
+
+	// Create a row for each attribute
+	for (TSharedPtr<FFunctionAttribute>& Attribute : RGBAttributes)
+	{
+		const TAttribute<EVisibility> VisibilityAttribute = TAttribute<EVisibility>(this, &FDMXPixelMappingDetailCustomization_FixtureGroupItem::GetRGBAttributesVisibility);
+		CreateAttributePropertyRowLambda(Attribute.ToSharedRef(), VisibilityAttribute);
+	}
+
+	for (TSharedPtr<FFunctionAttribute>& Attribute : MonochromeAttributes)
+	{
+		const TAttribute<EVisibility> VisibilityAttribute = TAttribute<EVisibility>(this, &FDMXPixelMappingDetailCustomization_FixtureGroupItem::GetMonochromeAttributesVisibility);
+		CreateAttributePropertyRowLambda(Attribute.ToSharedRef(), VisibilityAttribute);
+	}
+}
+
 void FDMXPixelMappingDetailCustomization_FixtureGroupItem::ForceRefresh()
 {
 	PropertyUtilities->ForceRefresh();
diff --git a/Engine/Plugins/VirtualProduction/DMX/DMXPixelMapping/Source/DMXPixelMappingEditor/Private/Customizations/DMXPixelMappingDetailCustomization_FixtureGroupItem.h b/Engine/Plugins/VirtualProduction/DMX/DMXPixelMapping/Source/DMXPixelMappingEditor/Private/Customizations/DMXPixelMappingDetailCustomization_FixtureGroupItem.h
--- a/Engine/Plugins/VirtualProduction/DMX/DMXPixelMapping/Source/DMXPixelMappingEditor/Private/Customizations/DMXPixelMappingDetailCustomization_FixtureGroupItem.h
+++ b/Engine/Plugins/VirtualProduction/DMX/DMXPixelMapping/Source/DMXPixelMappingEditor/Private/Customizations/DMXPixelMappingDetailCustomization_FixtureGroupItem.h
@@ -3,7 +3,9 @@
 #pragma once
 
 #include "CoreMinimal.h"
+#include "DMXAttribute.h"
 #include "IDetailCustomization.h"
+#include "PropertyHandle.h"
 #include "Misc/Attribute.h"
 #include "UObject/WeakObjectPtr.h"
 #include "IPropertyTypeCustomization.h"
@@ -31,8 +33,16 @@ private:
 		TSharedPtr<IPropertyHandle> ExposeHandle;
 		TSharedPtr<IPropertyHandle> InvertHandle;
 		TAttribute<EVisibility> Visibility;
-	};
 
+		/** Returns true if the Attribute Handle has multiple values */
+		bool HasMultipleAttributeValues() const;
+
+		/** Gets the value of the Attribute Handle. Should not be called if HasMultipleAttributeValues returns true */
+		FName GetAttributeValue() const;
+
+		/** Sets the value of the Attribute Handle */
+		void SetAttributeValue(const FName& NewValue);
+	};
 
 public:
 	static TSharedRef<IDetailCustomization> MakeInstance(const TWeakPtr<FDMXPixelMappingToolkit> InToolkitWeakPtr)
@@ -50,7 +60,6 @@ public:
 	//~ IPropertyTypeCustomization interface end
 
 private:
-
 	EVisibility GetRGBAttributeRowVisibilty(FFunctionAttribute* Attribute) const;
 
 	EVisibility GetRGBAttributesVisibility() const;
@@ -67,6 +76,9 @@ private:
 	/** Creates Details for the Output Modulators */
 	void CreateModulatorDetails(IDetailLayoutBuilder& InDetailLayout);
 
+	/** Creates the Details for the Attributes */
+	void CreateAttributeDetails(IDetailLayoutBuilder& InDetailLayout);
+
 	/** Forces the layout to redraw */
 	void ForceRefresh();
 
diff --git a/Engine/Plugins/VirtualProduction/DMX/DMXPixelMapping/Source/DMXPixelMappingEditor/Private/Customizations/DMXPixelMappingDetailCustomization_Matrix.cpp b/Engine/Plugins/VirtualProduction/DMX/DMXPixelMapping/Source/DMXPixelMappingEditor/Private/Customizations/DMXPixelMappingDetailCustomization_Matrix.cpp
--- a/Engine/Plugins/VirtualProduction/DMX/DMXPixelMapping/Source/DMXPixelMappingEditor/Private/Customizations/DMXPixelMappingDetailCustomization_Matrix.cpp
+++ b/Engine/Plugins/VirtualProduction/DMX/DMXPixelMapping/Source/DMXPixelMappingEditor/Private/Customizations/DMXPixelMappingDetailCustomization_Matrix.cpp
@@ -12,6 +12,7 @@
 #include "Library/DMXEntityFixturePatch.h"
 #include "Modules/ModuleManager.h"
 #include "Toolkits/DMXPixelMappingToolkit.h"
+#include "Widgets/SDMXPixelMappingAttributeNamesComboBox.h"
 
 #include "DetailWidgetRow.h"
 #include "DetailLayoutBuilder.h"
@@ -25,6 +26,68 @@
 
 #define LOCTEXT_NAMESPACE "DMXPixelMappingDetailCustomization_Matrix"
 
+bool FDMXPixelMappingDetailCustomization_Matrix::FDMXCellAttributeGroup::HasMultipleAttributeValues() const
+{
+	TArray<const void*> RawData;
+	Handle->AccessRawData(RawData);
+
+	TSet<FName> AttributeNames;
+	for (const void* RawPtr : RawData)
+	{
+		if (RawPtr)
+		{
+			// The types we use with this customization must have a cast constructor to FName
+			AttributeNames.Add(reinterpret_cast<const FDMXAttributeName*>(RawPtr)->Name);
+		}
+	}
+
+	return AttributeNames.Num() > 1;
+}
+
+FName FDMXPixelMappingDetailCustomization_Matrix::FDMXCellAttributeGroup::GetAttributeValue() const
+{
+	if (!ensureMsgf(!HasMultipleAttributeValues(), TEXT("Cannot get attribute value from handle when handle has mutliple values")))
+	{
+		return NAME_None;
+	}
+
+	TArray<const void*> RawData;
+	Handle->AccessRawData(RawData);
+
+	for (const void* RawPtr : RawData)
+	{
+		if (RawPtr)
+		{
+			// The types we use with this customization must have a cast constructor to FName
+			return reinterpret_cast<const FDMXAttributeName*>(RawPtr)->Name;
+		}
+	}
+
+	return NAME_None;
+}
+
+void FDMXPixelMappingDetailCustomization_Matrix::FDMXCellAttributeGroup::SetAttributeValue(const FName& NewValue)
+{
+	FStructProperty* StructProperty = CastFieldChecked<FStructProperty>(Handle->GetProperty());
+
+	TArray<void*> RawData;
+	Handle->AccessRawData(RawData);
+
+	for (void* SingleRawData : RawData)
+	{
+		FDMXAttributeName* PreviousValue = reinterpret_cast<FDMXAttributeName*>(SingleRawData);
+		FDMXAttributeName NewAttributeName;
+		NewAttributeName.SetFromName(NewValue);
+
+		// Export new value to text format that can be imported later
+		FString TextValue;
+		StructProperty->Struct->ExportText(TextValue, &NewAttributeName, PreviousValue, nullptr, EPropertyPortFlags::PPF_None, nullptr);
+
+		// Set values on edited property handle from exported text
+		ensure(Handle->SetValueFromFormattedString(TextValue, EPropertyValueSetFlags::DefaultFlags) == FPropertyAccess::Result::Success);
+	}
+}
+
 void FDMXPixelMappingDetailCustomization_Matrix::CustomizeDetails(IDetailLayoutBuilder& InDetailLayout)
 {
 	DetailLayout = &InDetailLayout;
@@ -92,18 +155,6 @@ void FDMXPixelMappingDetailCustomization_Matrix::CustomizeDetails(IDetailLayoutB
 			.OnGenerateRow(this, &FDMXPixelMappingDetailCustomization_Matrix::GenerateExposeAndInvertRow)
 		];
 
-	// Update RGB attributes
-	for (TSharedPtr<FDMXCellAttributeGroup>& Attribute : RGBAttributes)
-	{
-		DetailLayout->HideProperty(Attribute->ExposeHandle);
-		DetailLayout->HideProperty(Attribute->InvertHandle);
-
-		OutputSettingsCategory
-			.AddProperty(Attribute->Handle)
-			.Visibility(TAttribute<EVisibility>::Create(TAttribute<EVisibility>::FGetter::CreateSP(this, &FDMXPixelMappingDetailCustomization_Matrix::GetRGBAttributeRowVisibilty, Attribute.Get())));
-	}
-
-
 	// Generate all Monochrome Expose and Invert rows
 	OutputSettingsCategory.AddCustomRow(FText::GetEmpty())
 		.Visibility(TAttribute<EVisibility>(this, &FDMXPixelMappingDetailCustomization_Matrix::GetMonochromeAttributesVisibility))
@@ -118,18 +169,8 @@ void FDMXPixelMappingDetailCustomization_Matrix::CustomizeDetails(IDetailLayoutB
 			.OnGenerateRow(this, &FDMXPixelMappingDetailCustomization_Matrix::GenerateExposeAndInvertRow)
 		];
 
-	// Update Monochrome attributes
-	for (TSharedPtr<FDMXCellAttributeGroup>& Attribute : MonochromeAttributes)
-	{
-		DetailLayout->HideProperty(Attribute->ExposeHandle);
-		DetailLayout->HideProperty(Attribute->InvertHandle);
-
-		OutputSettingsCategory
-			.AddProperty(Attribute->Handle)
-			.Visibility(TAttribute<EVisibility>::Create(TAttribute<EVisibility>::FGetter::CreateSP(this, &FDMXPixelMappingDetailCustomization_Matrix::GetMonochromeRowVisibilty, Attribute.Get())));
-	}
-	
 	CreateModulatorDetails(InDetailLayout);
+	CreateAttributeDetails(InDetailLayout);
 }
 
 bool FDMXPixelMappingDetailCustomization_Matrix::CheckComponentsDMXColorMode(const EDMXColorMode DMXColorMode) const
@@ -336,6 +377,135 @@ void FDMXPixelMappingDetailCustomization_Matrix::CreateModulatorDetails(IDetailL
 	}
 }
 
+void FDMXPixelMappingDetailCustomization_Matrix::CreateAttributeDetails(IDetailLayoutBuilder& InDetailLayout)
+{
+	TArray<TWeakObjectPtr<UObject>> Objects;
+	DetailLayout->GetObjectsBeingCustomized(Objects);
+
+	TArray<UDMXPixelMappingMatrixComponent*> StrongMatrixComponents;
+	for (TWeakObjectPtr<UObject> SelectedObject : Objects)
+	{
+		UDMXPixelMappingMatrixComponent* MatrixComponent = Cast<UDMXPixelMappingMatrixComponent>(SelectedObject);
+		if (MatrixComponent)
+		{
+			StrongMatrixComponents.Add(MatrixComponent);
+		}
+	}
+
+	// Gather attribute names present in Patches of all selected components
+	TArray<FName> FixtureMatrixAttributes;
+	for (UDMXPixelMappingMatrixComponent* MatrixComponent : StrongMatrixComponents)
+	{
+		const UDMXEntityFixturePatch* FixturePatch = MatrixComponent->FixturePatchRef.GetFixturePatch();
+		if (!FixturePatch)
+		{
+			continue;
+		}
+
+		const FDMXFixtureMode* ModePtr = FixturePatch->GetActiveMode();
+		if (!ModePtr)
+		{
+			continue;
+		}
+
+		const TArray<FDMXFixtureCellAttribute> CellAttributes = ModePtr->FixtureMatrixConfig.CellAttributes;
+		for (const FDMXFixtureCellAttribute& CellAttribute : CellAttributes)
+		{
+			FixtureMatrixAttributes.AddUnique(CellAttribute.Attribute.Name);
+		}
+	}
+
+	for (UDMXPixelMappingMatrixComponent* MatrixComponent : StrongMatrixComponents)
+	{
+		FixtureMatrixAttributes.RemoveAll([MatrixComponent](const FName& AttributeName)
+			{
+				const UDMXEntityFixturePatch* FixturePatch = MatrixComponent->FixturePatchRef.GetFixturePatch();
+				if (!FixturePatch)
+				{
+					return true;
+				}
+
+				const FDMXFixtureMode* ModePtr = FixturePatch->GetActiveMode();
+				if (!ModePtr)
+				{
+					return true;
+				}
+
+				const TArray<FDMXFixtureCellAttribute> AttributeNamesOfPatch = ModePtr->FixtureMatrixConfig.CellAttributes;
+				return !AttributeNamesOfPatch.ContainsByPredicate([AttributeName](const FDMXFixtureCellAttribute& AttributeNameOfPatch)
+					{
+						return AttributeNameOfPatch.Attribute == AttributeName;
+					});
+			});
+	}
+
+	// Lambda to add a single attribute property row
+	auto CreateAttributePropertyRowLambda = [&InDetailLayout, &FixtureMatrixAttributes](const TSharedRef<FDMXCellAttributeGroup>& Attribute, const TAttribute<EVisibility>& VisibilityAttribute)
+		{
+			InDetailLayout.HideProperty(Attribute->ExposeHandle);
+			InDetailLayout.HideProperty(Attribute->InvertHandle);
+			InDetailLayout.HideProperty(Attribute->Handle);
+
+			const TSharedRef<SDMXPixelMappingAttributeNamesComboBox> AttributeNameComboBox =
+				SNew(SDMXPixelMappingAttributeNamesComboBox, FixtureMatrixAttributes)
+				.OnSelectionChanged_Lambda([Attribute](const FName& NewValue)
+					{
+						Attribute->SetAttributeValue(NewValue);
+					});
+
+			if (Attribute->HasMultipleAttributeValues())
+			{
+				AttributeNameComboBox->SetHasMultipleValues();
+			}
+			else
+			{
+				FName InitialSelection = Attribute->GetAttributeValue();
+
+				// Set a valid attribute, or name none if none is available
+				if (!FixtureMatrixAttributes.Contains(InitialSelection))
+				{
+					if (FixtureMatrixAttributes.IsEmpty())
+					{
+						Attribute->SetAttributeValue(NAME_None);
+					}
+					else
+					{
+						Attribute->SetAttributeValue(FixtureMatrixAttributes[0]);
+						InitialSelection = FixtureMatrixAttributes[0];
+					}
+				}
+				AttributeNameComboBox->SetSelection(InitialSelection);
+			}
+
+			InDetailLayout.EditCategory("Output Settings", FText::GetEmpty(), ECategoryPriority::Important)
+				.AddCustomRow(FText::GetEmpty())
+				.Visibility(VisibilityAttribute)
+				.NameContent()
+				[
+					SNew(STextBlock)
+					.Text(Attribute->Handle->GetPropertyDisplayName())
+					.Font(FAppStyle::GetFontStyle(TEXT("PropertyWindow.NormalFont")))
+				]
+				.ValueContent()
+				[
+					AttributeNameComboBox
+				];
+		};
+
+	// Create a row for each attribute
+	for (TSharedPtr<FDMXCellAttributeGroup>& Attribute : RGBAttributes)
+	{
+		const TAttribute<EVisibility> VisibilityAttribute = TAttribute<EVisibility>(this, &FDMXPixelMappingDetailCustomization_Matrix::GetRGBAttributesVisibility);
+		CreateAttributePropertyRowLambda(Attribute.ToSharedRef(), VisibilityAttribute);
+	}
+
+	for (TSharedPtr<FDMXCellAttributeGroup>& Attribute : MonochromeAttributes)
+	{
+		const TAttribute<EVisibility> VisibilityAttribute = TAttribute<EVisibility>(this, &FDMXPixelMappingDetailCustomization_Matrix::GetMonochromeAttributesVisibility);
+		CreateAttributePropertyRowLambda(Attribute.ToSharedRef(), VisibilityAttribute);
+	}
+}
+
 void FDMXPixelMappingDetailCustomization_Matrix::ForceRefresh()
 {
 	PropertyUtilities->ForceRefresh();
diff --git a/Engine/Plugins/VirtualProduction/DMX/DMXPixelMapping/Source/DMXPixelMappingEditor/Private/Customizations/DMXPixelMappingDetailCustomization_Matrix.h b/Engine/Plugins/VirtualProduction/DMX/DMXPixelMapping/Source/DMXPixelMappingEditor/Private/Customizations/DMXPixelMappingDetailCustomization_Matrix.h
--- a/Engine/Plugins/VirtualProduction/DMX/DMXPixelMapping/Source/DMXPixelMappingEditor/Private/Customizations/DMXPixelMappingDetailCustomization_Matrix.h
+++ b/Engine/Plugins/VirtualProduction/DMX/DMXPixelMapping/Source/DMXPixelMappingEditor/Private/Customizations/DMXPixelMappingDetailCustomization_Matrix.h
@@ -31,6 +31,15 @@ private:
 		TSharedPtr<IPropertyHandle> ExposeHandle;
 		TSharedPtr<IPropertyHandle> InvertHandle;
 		TAttribute<EVisibility> Visibility;
+
+		/** Returns true if the Attribute Handle has multiple values */
+		bool HasMultipleAttributeValues() const;
+
+		/** Gets the value of the Attribute Handle. Should not be called if HasMultipleAttributeValues returns true */
+		FName GetAttributeValue() const;
+
+		/** Sets the value of the Attribute Handle */
+		void SetAttributeValue(const FName& NewValue);
 	};
 
 public:
@@ -65,6 +74,9 @@ private:
 	/** Creates Details for the Output Modulators */
 	void CreateModulatorDetails(IDetailLayoutBuilder& InDetailLayout);
 
+	/** Creates the Details for the Attributes */
+	void CreateAttributeDetails(IDetailLayoutBuilder& InDetailLayout);
+
 	/** Forces the layout to redraw */
 	void ForceRefresh();
 
diff --git a/Engine/Plugins/VirtualProduction/DMX/DMXPixelMapping/Source/DMXPixelMappingEditor/Private/Widgets/SDMXPixelMappingAttributeNamesComboBox.cpp b/Engine/Plugins/VirtualProduction/DMX/DMXPixelMapping/Source/DMXPixelMappingEditor/Private/Widgets/SDMXPixelMappingAttributeNamesComboBox.cpp
new file mode 100644
--- /dev/null
+++ b/Engine/Plugins/VirtualProduction/DMX/DMXPixelMapping/Source/DMXPixelMappingEditor/Private/Widgets/SDMXPixelMappingAttributeNamesComboBox.cpp
@@ -0,0 +1,88 @@
+// Copyright Epic Games, Inc. All Rights Reserved.
+
+#include "Widgets/SDMXPixelMappingAttributeNamesComboBox.h"
+
+#include "Algo/Find.h"
+#include "Widgets/Input/SComboBox.h"
+#include "Widgets/Text/STextBlock.h"
+
+
+#define LOCTEXT_NAMESPACE "SDMXPixelMappingAttributeNamesComboBox"
+
+void SDMXPixelMappingAttributeNamesComboBox::Construct(const FArguments& InArgs, const TArray<FName>& InOptions)
+{
+	for (const FName& AttributeName : InOptions)
+	{
+		Options.Add(MakeShared<FName>(AttributeName));
+	}
+
+	OnSelectionChanged = InArgs._OnSelectionChanged;
+
+	// Always have at least one option to display
+	if (Options.IsEmpty())
+	{
+		Options.Add(MakeShared<FName>("No attributes available"));
+	}
+
+	ChildSlot
+		[
+			SAssignNew(ComboBox, SComboBox<TSharedPtr<FName>>)
+			.OptionsSource(&Options)
+			.OnSelectionChanged(this, &SDMXPixelMappingAttributeNamesComboBox::HandleSelectionChanged)
+			.OnGenerateWidget(this, &SDMXPixelMappingAttributeNamesComboBox::OnGenerateComboBoxEntry)
+			[
+				SNew(STextBlock)
+				.Text_Lambda([this]()
+					{
+						if (bHasMultipleValues)
+						{
+							return LOCTEXT("MultipleValuesText", "Multiple Values");
+						}
+						else
+						{
+							check(!Options.IsEmpty());
+
+							const TSharedPtr<FName> SelectedItem = ComboBox->GetSelectedItem();
+							return SelectedItem.IsValid() ? FText::FromString(SelectedItem->ToString()) : FText::FromString(TEXT("Invalid selection"));
+						}
+					})
+				.Font(FAppStyle::GetFontStyle(TEXT("PropertyWindow.NormalFont")))
+			]
+		];
+}
+
+void SDMXPixelMappingAttributeNamesComboBox::SetSelection(const FName& Selection)
+{
+	bHasMultipleValues = false;
+	const TSharedPtr<FName>* SelectionPtr = Algo::FindByPredicate(Options, [&Selection](const TSharedPtr<FName>& AttributeName)
+		{
+			return *AttributeName == Selection;
+		});
+
+	ensureMsgf(SelectionPtr, TEXT("Cannot set selection of SDMXPixelMappingAttributeNamesComboBox. Selection %s is not provided as an option."), *Selection.ToString());
+	if (SelectionPtr)
+	{
+		ComboBox->SetSelectedItem(*SelectionPtr);
+	}
+}
+
+void SDMXPixelMappingAttributeNamesComboBox::SetHasMultipleValues()
+{
+	bHasMultipleValues = true;
+}
+
+TSharedRef<SWidget> SDMXPixelMappingAttributeNamesComboBox::OnGenerateComboBoxEntry(TSharedPtr<FName> AttributeName) const
+{
+	return SNew(STextBlock)
+		.Text(FText::FromString(AttributeName->ToString()))
+		.Font(FAppStyle::GetFontStyle(TEXT("PropertyWindow.NormalFont")));
+}
+
+void SDMXPixelMappingAttributeNamesComboBox::HandleSelectionChanged(TSharedPtr<FName> NewValue, ESelectInfo::Type)
+{
+	bHasMultipleValues = false;
+
+	OnSelectionChanged.ExecuteIfBound(*NewValue.Get());
+}
+
+#undef LOCTEXT_NAMESPACE
diff --git a/Engine/Plugins/VirtualProduction/DMX/DMXPixelMapping/Source/DMXPixelMappingEditor/Private/Widgets/SDMXPixelMappingAttributeNamesComboBox.h b/Engine/Plugins/VirtualProduction/DMX/DMXPixelMapping/Source/DMXPixelMappingEditor/Private/Widgets/SDMXPixelMappingAttributeNamesComboBox.h
new file mode 100644
--- /dev/null
+++ b/Engine/Plugins/VirtualProduction/DMX/DMXPixelMapping/Source/DMXPixelMappingEditor/Private/Widgets/SDMXPixelMappingAttributeNamesComboBox.h
@@ -0,0 +1,60 @@
+// Copyright Epic Games, Inc. All Rights Reserved.
+#pragma once
+
+#include "Widgets/SCompoundWidget.h"
+
+template<typename OptionType>
+class SComboBox;
+
+
+DECLARE_DELEGATE_OneParam(FOnAttributeSelectionChanged, const FName& /** AttributeName */)
+
+class SDMXPixelMappingAttributeNamesComboBox
+	: public SCompoundWidget
+{
+public:
+	SLATE_BEGIN_ARGS(SDMXPixelMappingAttributeNamesComboBox)
+	{}
+
+		/** Event raised when the selction changed */
+		SLATE_EVENT(FOnAttributeSelectionChanged, OnSelectionChanged)
+
+	SLATE_END_ARGS()
+
+	/** Constructs this widget */
+	void Construct(const FArguments& InArgs, const TArray<FName>& InOptions);
+
+	/** Sets the selection. Fails silently if the Name is not contained in options provided during construction */
+	void SetSelection(const FName& Selection);
+
+	/** Sets the combo box to display 'Multiple Values'. Note, this is reset when the combo box selects a value. */
+	void SetHasMultipleValues();
+
+private:
+	/*
+	 * Make a widget for the given option
+	 *
+	 * @param InOption   An option from which to construct the widget
+	 */
+	TSharedRef<SWidget> OnGenerateComboBoxEntry(TSharedPtr<FName> AttributeName) const;
+
+	/*
+	 * Handle current item selection change
+	 *
+	 * @param NewValue   A value to set as current item
+	 */
+	void HandleSelectionChanged(TSharedPtr<FName> NewValue, ESelectInfo::Type);
+
+private:
+	/** The actual combo box */
+	TSharedPtr<SComboBox<TSharedPtr<FName>>> ComboBox;
+
+	/** Options source for the combo box */
+	TArray<TSharedPtr<FName>> Options;
+
+	/** If true, shows 'Multiple Values' instead of the selection */
+	bool bHasMultipleValues = false;
+
+	// Slate Args
+	FOnAttributeSelectionChanged OnSelectionChanged;
+};
