DMX - Fix interactive changes in layout scripts turn the UI irresponsive, remove obsolete checkbox 'Make Squares' for Grid Layout Script, fix bad file name of SDMXPixelMappingLayoutView.cpp

diff --git a/Engine/Plugins/VirtualProduction/DMX/DMXPixelMapping/Source/DMXPixelMappingEditor/Private/ViewModels/DMXPixelMappingLayoutViewModel.cpp b/Engine/Plugins/VirtualProduction/DMX/DMXPixelMapping/Source/DMXPixelMappingEditor/Private/ViewModels/DMXPixelMappingLayoutViewModel.cpp
--- a/Engine/Plugins/VirtualProduction/DMX/DMXPixelMapping/Source/DMXPixelMappingEditor/Private/ViewModels/DMXPixelMappingLayoutViewModel.cpp
+++ b/Engine/Plugins/VirtualProduction/DMX/DMXPixelMapping/Source/DMXPixelMappingEditor/Private/ViewModels/DMXPixelMappingLayoutViewModel.cpp
@@ -131,8 +131,30 @@ TArray<UObject*> UDMXPixelMappingLayoutViewModel::GetLayoutScriptsObjectsSlow()
 	return Result;
 }
 
+void UDMXPixelMappingLayoutViewModel::RequestApplyLayoutScripts()
+{
+	if (!ApplyLayoutScriptTimerHandle.IsValid())
+	{
+		ApplyLayoutScriptTimerHandle = GEditor->GetTimerManager()->SetTimerForNextTick(FTimerDelegate::CreateUObject(this, &UDMXPixelMappingLayoutViewModel::ApplyLayoutScripts));
+	}
+}
+
+void UDMXPixelMappingLayoutViewModel::PostUndo(bool bSuccess)
+{
+	RefreshComponents();
+	RefreshLayoutScriptClass();
+}
+
+void UDMXPixelMappingLayoutViewModel::PostRedo(bool bSuccess)
+{
+	RefreshComponents();
+	RefreshLayoutScriptClass();
+}
+
 void UDMXPixelMappingLayoutViewModel::ApplyLayoutScripts()
 {
+	ApplyLayoutScriptTimerHandle.Invalidate();
+
 	const EDMXPixelMappingLayoutViewModelMode LayoutType = GetMode();
 	if (LayoutType == EDMXPixelMappingLayoutViewModelMode::LayoutRendererComponentChildren)
 	{
@@ -152,18 +174,6 @@ void UDMXPixelMappingLayoutViewModel::ApplyLayoutScripts()
 	}
 }
 
-void UDMXPixelMappingLayoutViewModel::PostUndo(bool bSuccess)
-{
-	RefreshComponents();
-	RefreshLayoutScriptClass();
-}
-
-void UDMXPixelMappingLayoutViewModel::PostRedo(bool bSuccess)
-{
-	RefreshComponents();
-	RefreshLayoutScriptClass();
-}
-
 void UDMXPixelMappingLayoutViewModel::LayoutRendererComponentChildren()
 {
 	const EDMXPixelMappingLayoutViewModelMode LayoutType = GetMode();
diff --git a/Engine/Plugins/VirtualProduction/DMX/DMXPixelMapping/Source/DMXPixelMappingEditor/Private/ViewModels/DMXPixelMappingLayoutViewModel.h b/Engine/Plugins/VirtualProduction/DMX/DMXPixelMapping/Source/DMXPixelMappingEditor/Private/ViewModels/DMXPixelMappingLayoutViewModel.h
--- a/Engine/Plugins/VirtualProduction/DMX/DMXPixelMapping/Source/DMXPixelMappingEditor/Private/ViewModels/DMXPixelMappingLayoutViewModel.h
+++ b/Engine/Plugins/VirtualProduction/DMX/DMXPixelMapping/Source/DMXPixelMappingEditor/Private/ViewModels/DMXPixelMappingLayoutViewModel.h
@@ -4,6 +4,7 @@
 
 #include "CoreMinimal.h"
 #include "EditorUndoClient.h"
+#include "TimerManager.h"
 
 #include "DMXPixelMappingLayoutViewModel.generated.h"
 
@@ -57,8 +58,8 @@ public:
 	/** Returns the Layout Scripts of relevant components. Looks up each relevant component and returns its script. */
 	TArray<UObject*> GetLayoutScriptsObjectsSlow() const;
 
-	/** Applies the layouts scripts currently in use */
-	void ApplyLayoutScripts();
+	/** Applies the layouts scripts currently in use on the next tick */
+	void RequestApplyLayoutScripts();
 
 	/** Delegate executed when the Model changed */
 	FSimpleMulticastDelegate OnModelChanged;
@@ -73,6 +74,9 @@ protected:
 	// End of FEditorUndoClient
 
 private:
+	/** Applies the layouts scripts currently in use */
+	void ApplyLayoutScripts();
+
 	/** Lays out children of the Renderer Component */
 	void LayoutRendererComponentChildren();
 
@@ -126,6 +130,9 @@ private:
 	UPROPERTY(EditAnywhere, Category = "Layout Script", Meta = (ShowDisplayNames))
 	TSoftClassPtr<UDMXPixelMappingLayoutScript> LayoutScriptClass;
 
+	/** Timer handle to apply the layout script, set when ApplyLayoutScript is requested. */
+	FTimerHandle ApplyLayoutScriptTimerHandle;
+
 	/** The Layout Script class currently before it was changed */
 	TSoftClassPtr<UDMXPixelMappingLayoutScript> PreEditChangeLayoutScriptClass;
 };
diff --git a/Engine/Plugins/VirtualProduction/DMX/DMXPixelMapping/Source/DMXPixelMappingEditor/Private/Views/SDMXPixelMappingLayoutView..cpp b/Engine/Plugins/VirtualProduction/DMX/DMXPixelMapping/Source/DMXPixelMappingEditor/Private/Views/SDMXPixelMappingLayoutView..cpp
deleted file mode 100644
diff --git a/Engine/Plugins/VirtualProduction/DMX/DMXPixelMapping/Source/DMXPixelMappingEditor/Private/Views/SDMXPixelMappingLayoutView.cpp b/Engine/Plugins/VirtualProduction/DMX/DMXPixelMapping/Source/DMXPixelMappingEditor/Private/Views/SDMXPixelMappingLayoutView.cpp
new file mode 100644
--- /dev/null
+++ b/Engine/Plugins/VirtualProduction/DMX/DMXPixelMapping/Source/DMXPixelMappingEditor/Private/Views/SDMXPixelMappingLayoutView.cpp
@@ -0,0 +1,166 @@
+// Copyright Epic Games, Inc. All Rights Reserved.
+
+#include "Views/SDMXPixelMappingLayoutView.h"
+
+#include "DMXPixelMappingLayoutSettings.h"
+#include "Customizations/DMXPixelMappingHorizontalAlignmentCustomization.h"
+#include "Customizations/DMXPixelMappingLayoutViewModelDetails.h"
+#include "Customizations/DMXPixelMappingVerticalAlignmentCustomization.h"
+#include "Toolkits/DMXPixelMappingToolkit.h"
+#include "ViewModels/DMXPixelMappingLayoutViewModel.h"
+
+#include "DetailsViewArgs.h"
+#include "Editor.h"
+#include "IDetailsView.h"
+#include "PropertyEditorModule.h"
+#include "TimerManager.h"
+#include "Modules/ModuleManager.h"
+#include "Widgets/Input/SCheckBox.h"
+#include "Widgets/Layout/SBorder.h"
+#include "Widgets/Layout/SScrollBox.h"
+
+
+#define LOCTEXT_NAMESPACE "SDMXPixelMappingHierarchyView"
+
+SDMXPixelMappingLayoutView::~SDMXPixelMappingLayoutView()
+{
+	if (GEditor)
+	{
+		GEditor->UnregisterForUndo(this);
+	}
+}
+
+void SDMXPixelMappingLayoutView::Construct(const FArguments& InArgs, const TSharedRef<FDMXPixelMappingToolkit>& InToolkit)
+{
+	WeakToolkit = InToolkit;
+
+	Model = NewObject<UDMXPixelMappingLayoutViewModel>();
+	Model->SetToolkit(InToolkit);
+	Model->OnModelChanged.AddSP(this, &SDMXPixelMappingLayoutView::Refresh);
+
+	InitializeLayoutScriptsDetailsView();
+
+	if (!ensureMsgf(LayoutScriptDetailsView.IsValid(), TEXT("Failed to create Details View for Pixel Mapping Layouts")))
+	{
+		return;
+	}
+
+	ChildSlot
+		[
+			SNew(SScrollBox)
+
+			// Layout Model Details (select Layout Script class etc.)
+			+ SScrollBox::Slot()
+			[
+				CreateLayoutModelDetailsView()
+			]
+
+			// Layout Script Details
+			+ SScrollBox::Slot()
+			[
+				LayoutScriptDetailsView.ToSharedRef()
+			]
+
+			// Layout Settings Details
+			+ SScrollBox::Slot()
+			[
+				CreateLayoutSettingsDetailsView()
+			]
+
+		];
+
+	Refresh();
+
+	if (GEditor)
+	{
+		GEditor->RegisterForUndo(this);
+	}
+}
+
+void SDMXPixelMappingLayoutView::AddReferencedObjects(FReferenceCollector& Collector)
+{
+	Collector.AddReferencedObject(Model);
+}
+
+void SDMXPixelMappingLayoutView::PostUndo(bool bSuccess)
+{
+	GEditor->GetTimerManager()->SetTimerForNextTick(FTimerDelegate::CreateSP(this, &SDMXPixelMappingLayoutView::Refresh));
+}
+
+void SDMXPixelMappingLayoutView::PostRedo(bool bSuccess)
+{
+	GEditor->GetTimerManager()->SetTimerForNextTick(FTimerDelegate::CreateSP(this, &SDMXPixelMappingLayoutView::Refresh));
+}
+
+void SDMXPixelMappingLayoutView::NotifyPostChange(const FPropertyChangedEvent& PropertyChangedEvent, FProperty* PropertyThatChanged)
+{
+	Model->RequestApplyLayoutScripts();
+}
+
+void SDMXPixelMappingLayoutView::Refresh()
+{
+	const TArray<UObject*> LayoutScriptObjects = Model->GetLayoutScriptsObjectsSlow();
+	if (LayoutScriptObjects.IsEmpty())
+	{
+		LayoutScriptDetailsView->SetVisibility(EVisibility::Collapsed);
+	}
+	else
+	{
+		LayoutScriptDetailsView->SetObjects(LayoutScriptObjects);
+		LayoutScriptDetailsView->ForceRefresh();
+
+		LayoutScriptDetailsView->SetVisibility(EVisibility::Visible);
+	}
+}
+
+TSharedRef<SWidget> SDMXPixelMappingLayoutView::CreateLayoutSettingsDetailsView()
+{
+	UDMXPixelMappingLayoutSettings* LayoutSettings = GetMutableDefault<UDMXPixelMappingLayoutSettings>();
+	if (!LayoutSettings)
+	{
+		return SNullWidget::NullWidget;
+	}
+
+	FPropertyEditorModule& PropertyEditorModule = FModuleManager::GetModuleChecked<FPropertyEditorModule>("PropertyEditor");
+	FDetailsViewArgs DetailsViewArgs;
+	DetailsViewArgs.bAllowSearch = false;
+	DetailsViewArgs.NameAreaSettings = FDetailsViewArgs::HideNameArea;
+	TSharedRef<IDetailsView> LayoutSettingsDetailsView = PropertyEditorModule.CreateDetailView(DetailsViewArgs);
+
+	LayoutSettingsDetailsView->SetObject(LayoutSettings);
+
+	return LayoutSettingsDetailsView;
+}
+
+TSharedRef<SWidget> SDMXPixelMappingLayoutView::CreateLayoutModelDetailsView()
+{
+	FPropertyEditorModule& PropertyEditorModule = FModuleManager::GetModuleChecked<FPropertyEditorModule>("PropertyEditor");
+	FDetailsViewArgs DetailsViewArgs;
+	DetailsViewArgs.bAllowSearch = false;
+	DetailsViewArgs.NameAreaSettings = FDetailsViewArgs::HideNameArea;
+	TSharedRef<IDetailsView> ModelDetailsView = PropertyEditorModule.CreateDetailView(DetailsViewArgs);
+	
+	ModelDetailsView->RegisterInstancedCustomPropertyLayout(UDMXPixelMappingLayoutViewModel::StaticClass(), FOnGetDetailCustomizationInstance::CreateStatic(&FDMXPixelMappingLayoutViewModelDetails::MakeInstance, WeakToolkit));
+
+	ModelDetailsView->SetObject(Model);
+
+	return ModelDetailsView;
+}
+
+void SDMXPixelMappingLayoutView::InitializeLayoutScriptsDetailsView()
+{
+	TArray<UObject*> LayoutScripts = Model->GetLayoutScriptsObjectsSlow();
+
+	FPropertyEditorModule& PropertyEditorModule = FModuleManager::GetModuleChecked<FPropertyEditorModule>("PropertyEditor");
+	FDetailsViewArgs DetailsViewArgs;
+	DetailsViewArgs.bAllowSearch = false;
+	DetailsViewArgs.NotifyHook = this;
+
+	DetailsViewArgs.NameAreaSettings = FDetailsViewArgs::HideNameArea;
+	LayoutScriptDetailsView = PropertyEditorModule.CreateDetailView(DetailsViewArgs);
+
+	LayoutScriptDetailsView->RegisterInstancedCustomPropertyTypeLayout(TEXT("EHorizontalAlignment"), FOnGetPropertyTypeCustomizationInstance::CreateStatic(&FDMXPixelMappingHorizontalAlignmentCustomization::MakeInstance));
+	LayoutScriptDetailsView->RegisterInstancedCustomPropertyTypeLayout(TEXT("EVerticalAlignment"), FOnGetPropertyTypeCustomizationInstance::CreateStatic(&FDMXPixelMappingVerticalAlignmentCustomization::MakeInstance));
+}
+
+#undef LOCTEXT_NAMESPACE
diff --git a/Engine/Plugins/VirtualProduction/DMX/DMXPixelMapping/Source/DMXPixelMappingRuntime/Private/LayoutScripts/DMXPixelMappingLayoutScript_GridLayout.h b/Engine/Plugins/VirtualProduction/DMX/DMXPixelMapping/Source/DMXPixelMappingRuntime/Private/LayoutScripts/DMXPixelMappingLayoutScript_GridLayout.h
--- a/Engine/Plugins/VirtualProduction/DMX/DMXPixelMapping/Source/DMXPixelMappingRuntime/Private/LayoutScripts/DMXPixelMappingLayoutScript_GridLayout.h
+++ b/Engine/Plugins/VirtualProduction/DMX/DMXPixelMapping/Source/DMXPixelMappingRuntime/Private/LayoutScripts/DMXPixelMappingLayoutScript_GridLayout.h
@@ -51,10 +51,6 @@ public:
 	UPROPERTY(EditAnywhere, Category = "Layout Settings", Meta = (DisplayName = "Vertical Alignment"))
 	TEnumAsByte<EVerticalAlignment> VerticalAlignment;
 
-	/** If true, cells are always shaped squares */
-	UPROPERTY(EditAnywhere, Category = "Layout Settings")
-	bool bMakeSquares = true;
-	
 	/** How cells are distributed, compared to their previous distribution */
 	UPROPERTY(EditAnywhere, Category = "Layout Settings")
 	EDMXPixelMappingDistribution Distribution = EDMXPixelMappingDistribution::TopLeftToRight;
