DMX - Modulators in PixelMapping show attributes available to the patch instead of default attributes

diff --git a/Engine/Plugins/VirtualProduction/DMX/DMXPixelMapping/Source/DMXPixelMappingEditor/Private/Customizations/DMXPixelMappingAttributeNamesDetails.cpp b/Engine/Plugins/VirtualProduction/DMX/DMXPixelMapping/Source/DMXPixelMappingEditor/Private/Customizations/DMXPixelMappingAttributeNamesDetails.cpp
new file mode 100644
--- /dev/null
+++ b/Engine/Plugins/VirtualProduction/DMX/DMXPixelMapping/Source/DMXPixelMappingEditor/Private/Customizations/DMXPixelMappingAttributeNamesDetails.cpp
@@ -0,0 +1,290 @@
+// Copyright Epic Games, Inc. All Rights Reserved.
+
+#include "Customizations/DMXPixelMappingAttributeNamesDetails.h"
+
+#include "DMXAttribute.h"
+#include "Library/DMXEntityFixturePatch.h"
+
+#include "DetailLayoutBuilder.h"
+#include "DetailWidgetRow.h"
+#include "PropertyHandle.h"
+#include "Algo/Find.h"
+#include "Widgets/SCompoundWidget.h"
+#include "Widgets/Input/SComboBox.h"
+#include "Widgets/Text/STextBlock.h"
+
+
+#define LOCTEXT_NAMESPACE "DMXPixelMappingAttributeNamesDetails"
+
+class SDMXPixelMappingAttributeNamesComboBox
+	: public SCompoundWidget
+{
+	DECLARE_DELEGATE_OneParam(FOnAttributeSelectionChanged, const FName& /** AttributeName */)
+
+public:
+	SLATE_BEGIN_ARGS(SDMXPixelMappingAttributeNamesComboBox)
+	{}
+
+		/** Event raised when the selection changed */
+		SLATE_EVENT(FOnAttributeSelectionChanged, OnSelectionChanged)
+
+	SLATE_END_ARGS()
+
+	/** Constructs this widget */
+	void Construct(const FArguments& InArgs, const TArray<FName>& InOptions)
+	{
+		for (const FName& AttributeName : InOptions)
+		{
+			Options.Add(MakeShared<FName>(AttributeName));
+		}
+
+		OnSelectionChanged = InArgs._OnSelectionChanged;
+
+		// Always have at least one option to display
+		if (Options.IsEmpty())
+		{
+			Options.Add(MakeShared<FName>("No attributes available"));
+		}
+
+		ChildSlot
+			[
+				SAssignNew(ComboBox, SComboBox<TSharedPtr<FName>>)
+				.OptionsSource(&Options)
+				.OnSelectionChanged(this, &SDMXPixelMappingAttributeNamesComboBox::HandleSelectionChanged)
+				.OnGenerateWidget(this, &SDMXPixelMappingAttributeNamesComboBox::OnGenerateComboBoxEntry)
+				[
+					SNew(STextBlock)
+					.Text_Lambda([this]()
+						{
+							if (bHasMultipleValues)
+							{
+								return LOCTEXT("MultipleValuesText", "Multiple Values");
+							}
+							else
+							{
+								check(!Options.IsEmpty());
+
+								const TSharedPtr<FName> SelectedItem = ComboBox->GetSelectedItem();
+								return SelectedItem.IsValid() ? FText::FromString(SelectedItem->ToString()) : FText::FromString(TEXT("Invalid selection"));
+							}
+						})
+					.Font(FAppStyle::GetFontStyle(TEXT("PropertyWindow.NormalFont")))
+				]
+			];
+	}
+
+	/** Sets the selection */
+	void SetSelection(const FName& Selection)
+	{
+		bHasMultipleValues = false;
+		const TSharedPtr<FName>* SelectionPtr = Algo::FindByPredicate(Options, [&Selection](const TSharedPtr<FName>& AttributeName)
+			{
+				return *AttributeName == Selection;
+			});
+
+		ensureMsgf(SelectionPtr, TEXT("Cannot set selection of SDMXPixelMappingAttributeNamesComboBox. Selection %s is not provided as an option."), *Selection.ToString());
+		if (SelectionPtr)
+		{
+			ComboBox->SetSelectedItem(*SelectionPtr);
+		}
+	}
+
+	/** Sets the combo box to display 'Multiple Values'. Note, this is reset when the combo box selects a value. */
+	void SetHasMultipleValues() { bHasMultipleValues = true; }
+
+private:
+	/*
+	 * Make a widget for the given option
+	 *
+	 * @param InOption   An option from which to construct the widget
+	 */
+	TSharedRef<SWidget> OnGenerateComboBoxEntry(TSharedPtr<FName> AttributeName) const
+	{
+		return SNew(STextBlock)
+			.Text(FText::FromString(AttributeName->ToString()))
+			.Font(FAppStyle::GetFontStyle(TEXT("PropertyWindow.NormalFont")));
+	}
+
+	/*
+	 * Handle current item selection change
+	 *
+	 * @param NewValue   A value to set as current item
+	 */
+	void HandleSelectionChanged(TSharedPtr<FName> NewValue, ESelectInfo::Type)
+	{
+		bHasMultipleValues = false;
+
+		OnSelectionChanged.ExecuteIfBound(*NewValue.Get());
+	}
+
+private:
+	/** The actual combo box */
+	TSharedPtr<SComboBox<TSharedPtr<FName>>> ComboBox;
+
+	/** Options source for the combo box */
+	TArray<TSharedPtr<FName>> Options;
+
+	/** If true, shows 'Multiple Values' instead of the selection */
+	bool bHasMultipleValues = false;
+
+	// Slate Args
+	FOnAttributeSelectionChanged OnSelectionChanged;
+};
+
+TSharedRef<IPropertyTypeCustomization> FDMXPixelMappingAttributeNamesDetails::MakeInstance(const TArray<TWeakObjectPtr<UDMXEntityFixturePatch>> FixturePatches)
+{
+	return MakeShared<FDMXPixelMappingAttributeNamesDetails>(FixturePatches);
+}
+
+void FDMXPixelMappingAttributeNamesDetails::CustomizeHeader(TSharedRef<IPropertyHandle> PropertyHandle, class FDetailWidgetRow& HeaderRow, IPropertyTypeCustomizationUtils& CustomizationUtils)
+{
+	FixtureGroupItemComponentsAttributes = GetFixtureGroupItemsAttributes(FixturePatches);
+
+	const TSharedRef<SDMXPixelMappingAttributeNamesComboBox> AttributeNamesComboBox =
+		SNew(SDMXPixelMappingAttributeNamesComboBox, FixtureGroupItemComponentsAttributes)
+		.OnSelectionChanged_Lambda([this, PropertyHandle](const FName& NewValue)
+			{
+				SetAttributeValue(PropertyHandle, NewValue);
+			});
+
+	if (HasMultipleAttributeValues(PropertyHandle))
+	{
+		AttributeNamesComboBox->SetHasMultipleValues();
+	}
+	else if (FixtureGroupItemComponentsAttributes.Num() > 0)
+	{
+		FName InitialSelection = GetAttributeValue(PropertyHandle);
+
+		// Set a valid attribute, or name none if none is available
+		if (!FixtureGroupItemComponentsAttributes.Contains(InitialSelection))
+		{
+			if (FixtureGroupItemComponentsAttributes.IsEmpty())
+			{
+				SetAttributeValue(PropertyHandle, NAME_None);
+			}
+			else
+			{
+				SetAttributeValue(PropertyHandle, FixtureGroupItemComponentsAttributes[0]);
+				InitialSelection = FixtureGroupItemComponentsAttributes[0];
+			}
+		}
+		AttributeNamesComboBox->SetSelection(InitialSelection);
+	}
+
+	HeaderRow
+		.NameContent()
+		[
+			SNew(STextBlock)
+			.Text(PropertyHandle->GetPropertyDisplayName())
+			.Font(FAppStyle::GetFontStyle(TEXT("PropertyWindow.NormalFont")))
+		]
+		.ValueContent()
+		[
+			AttributeNamesComboBox
+		];
+}
+
+void FDMXPixelMappingAttributeNamesDetails::CustomizeChildren(TSharedRef<IPropertyHandle> PropertyHandle, class IDetailChildrenBuilder& ChildBuilder, IPropertyTypeCustomizationUtils& CustomizationUtils)
+{
+	
+}
+
+bool FDMXPixelMappingAttributeNamesDetails::HasMultipleAttributeValues(TSharedRef<IPropertyHandle> AttributeHandle) const
+{
+	TArray<const void*> RawData;
+	AttributeHandle->AccessRawData(RawData);
+
+	TSet<FName> AttributeNames;
+	for (const void* RawPtr : RawData)
+	{
+		if (RawPtr)
+		{
+			AttributeNames.Add(reinterpret_cast<const FDMXAttributeName*>(RawPtr)->Name);
+		}
+	}
+
+	return AttributeNames.Num() > 1;
+}
+
+FName FDMXPixelMappingAttributeNamesDetails::GetAttributeValue(TSharedRef<IPropertyHandle> AttributeHandle) const
+{
+	if (!ensureMsgf(!HasMultipleAttributeValues(AttributeHandle), TEXT("Cannot get attribute value from handle when handle has mutliple values")))
+	{
+		return NAME_None;
+	}
+
+	TArray<const void*> RawData;
+	AttributeHandle->AccessRawData(RawData);
+
+	for (const void* RawPtr : RawData)
+	{
+		if (RawPtr)
+		{
+			return reinterpret_cast<const FDMXAttributeName*>(RawPtr)->Name;
+		}
+	}
+
+	return NAME_None;
+}
+
+void FDMXPixelMappingAttributeNamesDetails::SetAttributeValue(TSharedRef<IPropertyHandle> AttributeHandle, const FName& NewValue)
+{
+	FStructProperty* StructProperty = CastFieldChecked<FStructProperty>(AttributeHandle->GetProperty());
+
+	TArray<void*> RawData;
+	AttributeHandle->AccessRawData(RawData);
+
+	for (void* SingleRawData : RawData)
+	{
+		FDMXAttributeName* PreviousValue = reinterpret_cast<FDMXAttributeName*>(SingleRawData);
+		FDMXAttributeName NewAttributeName;
+		NewAttributeName.SetFromName(NewValue);
+
+		// Export new value to text format that can be imported later
+		FString TextValue;
+		StructProperty->Struct->ExportText(TextValue, &NewAttributeName, PreviousValue, nullptr, EPropertyPortFlags::PPF_None, nullptr);
+
+		// Set values on edited property handle from exported text
+		ensure(AttributeHandle->SetValueFromFormattedString(TextValue, EPropertyValueSetFlags::DefaultFlags) == FPropertyAccess::Result::Success);
+	}
+}
+
+TArray<FName> FDMXPixelMappingAttributeNamesDetails::GetFixtureGroupItemsAttributes(const TArray<TWeakObjectPtr<UDMXEntityFixturePatch>>& Patches)
+{
+	TArray<FName> ComponentsAttributes;
+
+	if (!Patches.IsEmpty())
+	{
+		// Gather attribute names present in Patches
+		for (TWeakObjectPtr<UDMXEntityFixturePatch> Patch : Patches)
+		{
+			if (!Patch.IsValid())
+			{
+				continue;
+			}
+
+			const TArray<FDMXAttributeName> AttributeNames = Patch->GetAllAttributesInActiveMode();
+			for (FDMXAttributeName AttributeName : AttributeNames)
+			{
+				FixtureGroupItemComponentsAttributes.AddUnique(AttributeName.Name);
+			}
+		}
+
+		for (TWeakObjectPtr<UDMXEntityFixturePatch> Patch : Patches)
+		{
+			FixtureGroupItemComponentsAttributes.RemoveAll([Patch](const FName& AttributeName)
+			{
+				if (!Patch.IsValid())
+				{
+					return true;
+				}
+				const TArray<FDMXAttributeName> AttributeNamesOfPatch = Patch->GetAllAttributesInActiveMode();
+				return !AttributeNamesOfPatch.Contains(AttributeName);
+			});
+		}
+	}
+
+	return FixtureGroupItemComponentsAttributes;
+}
+
+#undef LOCTEXT_NAMESPACE
diff --git a/Engine/Plugins/VirtualProduction/DMX/DMXPixelMapping/Source/DMXPixelMappingEditor/Private/Customizations/DMXPixelMappingAttributeNamesDetails.h b/Engine/Plugins/VirtualProduction/DMX/DMXPixelMapping/Source/DMXPixelMappingEditor/Private/Customizations/DMXPixelMappingAttributeNamesDetails.h
new file mode 100644
--- /dev/null
+++ b/Engine/Plugins/VirtualProduction/DMX/DMXPixelMapping/Source/DMXPixelMappingEditor/Private/Customizations/DMXPixelMappingAttributeNamesDetails.h
@@ -0,0 +1,51 @@
+// Copyright Epic Games, Inc. All Rights Reserved.
+
+#pragma once
+
+#include "CoreMinimal.h"
+#include "IPropertyTypeCustomization.h"
+#include "UObject/WeakObjectPtr.h"
+
+class IPropertyHandle;
+class IPropertyUtilities;
+class UDMXEntityFixturePatch;
+
+
+/** Property customization for Attributes in certain PixelMapping views that can provide a Fixture Patch */
+class FDMXPixelMappingAttributeNamesDetails
+	: public IPropertyTypeCustomization
+{
+public:
+	/** Creates an instance of this Detail Customization*/
+	static TSharedRef<IPropertyTypeCustomization> MakeInstance(const TArray<TWeakObjectPtr<UDMXEntityFixturePatch>> FixturePatches);
+
+	/** Constructor */
+	FDMXPixelMappingAttributeNamesDetails(const TArray<TWeakObjectPtr<UDMXEntityFixturePatch>> InFixturePatches)
+		: FixturePatches(InFixturePatches)
+	{};
+
+private:
+	//~Begin IPropertyTypeCustomization interface
+	virtual void CustomizeHeader(TSharedRef<IPropertyHandle> PropertyHandle, class FDetailWidgetRow& HeaderRow, IPropertyTypeCustomizationUtils& CustomizationUtils) override;
+	virtual void CustomizeChildren(TSharedRef<IPropertyHandle> PropertyHandle, class IDetailChildrenBuilder& Builder, IPropertyTypeCustomizationUtils& CustomizationUtils) override;
+	//~End IPropertyTypeCustomization interface
+
+	/** Returns true if the given Attribute Handle has multiple values */
+	bool HasMultipleAttributeValues(TSharedRef<IPropertyHandle> AttributeHandle) const;
+
+	/** Gets the value of the given Attribute Handle. Should not be called if HasMultipleAttributeValues returns true */
+	FName GetAttributeValue(TSharedRef<IPropertyHandle> AttributeHandle) const;
+
+	/** Sets the value of the given Attribute Handle */
+	void SetAttributeValue(TSharedRef<IPropertyHandle> AttributeHandle, const FName& NewValue);
+
+	/** Gets an array of Attribute names from the given array of Components */
+	TArray<FName> GetFixtureGroupItemsAttributes(const TArray<TWeakObjectPtr<UDMXEntityFixturePatch>>& Patches);
+
+private:
+	/** The Fixture Patches from which the attributes are generated */
+	TArray<TWeakObjectPtr<UDMXEntityFixturePatch>> FixturePatches;
+
+	/** The attributes available to this Fixture Group Item */
+	TArray<FName> FixtureGroupItemComponentsAttributes;
+};
diff --git a/Engine/Plugins/VirtualProduction/DMX/DMXPixelMapping/Source/DMXPixelMappingEditor/Private/Customizations/DMXPixelMappingDetailCustomization_FixtureGroupItem.cpp b/Engine/Plugins/VirtualProduction/DMX/DMXPixelMapping/Source/DMXPixelMappingEditor/Private/Customizations/DMXPixelMappingDetailCustomization_FixtureGroupItem.cpp
--- a/Engine/Plugins/VirtualProduction/DMX/DMXPixelMapping/Source/DMXPixelMappingEditor/Private/Customizations/DMXPixelMappingDetailCustomization_FixtureGroupItem.cpp
+++ b/Engine/Plugins/VirtualProduction/DMX/DMXPixelMapping/Source/DMXPixelMappingEditor/Private/Customizations/DMXPixelMappingDetailCustomization_FixtureGroupItem.cpp
@@ -2,13 +2,14 @@
 
 #include "Customizations/DMXPixelMappingDetailCustomization_FixtureGroupItem.h"
 
+#include "DMXAttribute.h"
 #include "DMXPixelMappingTypes.h"
 #include "Components/DMXPixelMappingFixtureGroupItemComponent.h"
+#include "Customizations/DMXPixelMappingAttributeNamesDetails.h"
 #include "Library/DMXEntityFixturePatch.h"
 #include "Library/DMXEntityReference.h"
 #include "Modulators/DMXModulator.h"
 #include "Toolkits/DMXPixelMappingToolkit.h"
-#include "Widgets/SDMXPixelMappingAttributeNamesComboBox.h"
 
 #include "DetailLayoutBuilder.h"
 #include "DetailCategoryBuilder.h"
@@ -20,7 +21,6 @@
 #include "Modules/ModuleManager.h"
 
 
-
 #define LOCTEXT_NAMESPACE "FixtureGroupItem"
 
 bool FDMXPixelMappingDetailCustomization_FixtureGroupItem::FFunctionAttribute::HasMultipleAttributeValues() const
@@ -99,25 +99,8 @@ void FDMXPixelMappingDetailCustomization_FixtureGroupItem::CustomizeDetails(IDet
 	{
 		FixtureGroupItemComponents.Add(Cast<UDMXPixelMappingFixtureGroupItemComponent>(SelectedObject));
 	}
-
-	// Get editing categories
-	IDetailCategoryBuilder& OutputSettingsCategory = DetailLayout->EditCategory("Output Settings", FText::GetEmpty(), ECategoryPriority::Important);
-
-	// Hide absolute postition property handles
-	TSharedPtr<IPropertyHandle> PositionXPropertyHandle = InDetailLayout.GetProperty(UDMXPixelMappingOutputComponent::GetPositionXPropertyName());
-	InDetailLayout.HideProperty(PositionXPropertyHandle);
-	TSharedPtr<IPropertyHandle> PositionYPropertyHandle = InDetailLayout.GetProperty(UDMXPixelMappingOutputComponent::GetPositionYPropertyName());
-	InDetailLayout.HideProperty(PositionYPropertyHandle);
-	TSharedPtr<IPropertyHandle> SizeXPropertyHandle = InDetailLayout.GetProperty(UDMXPixelMappingOutputComponent::GetSizeXPropertyName());
-	InDetailLayout.HideProperty(SizeXPropertyHandle);
-	TSharedPtr<IPropertyHandle> SizeYPropertyHandle = InDetailLayout.GetProperty(UDMXPixelMappingOutputComponent::GetSizeYPropertyName());
-	InDetailLayout.HideProperty(SizeXPropertyHandle);
-
-	// Add Function and ColorMode properties at the beginning
-	TSharedPtr<IPropertyHandle> ColorModePropertyHandle = DetailLayout->GetProperty(GET_MEMBER_NAME_CHECKED(UDMXPixelMappingFixtureGroupItemComponent, ColorMode));
-	OutputSettingsCategory.AddProperty(ColorModePropertyHandle);
-
-	// Register attributes
+	
+	// Gather attribute property handles
 	TSharedPtr<FFunctionAttribute> AttributeR = MakeShared<FFunctionAttribute>();
 	AttributeR->Handle = DetailLayout->GetProperty(GET_MEMBER_NAME_CHECKED(UDMXPixelMappingFixtureGroupItemComponent, AttributeR));
 	AttributeR->ExposeHandle = DetailLayout->GetProperty(GET_MEMBER_NAME_CHECKED(UDMXPixelMappingFixtureGroupItemComponent, AttributeRExpose));
@@ -137,14 +120,40 @@ void FDMXPixelMappingDetailCustomization_FixtureGroupItem::CustomizeDetails(IDet
 	RGBAttributes.Add(AttributeG);
 	RGBAttributes.Add(AttributeB);
 
-	// Register Monochrome attribute
 	TSharedPtr<FFunctionAttribute> MonochromeAttribute = MakeShared<FFunctionAttribute>();
 	MonochromeAttribute->Handle = DetailLayout->GetProperty(GET_MEMBER_NAME_CHECKED(UDMXPixelMappingFixtureGroupItemComponent, MonochromeIntensity));
 	MonochromeAttribute->ExposeHandle = DetailLayout->GetProperty(GET_MEMBER_NAME_CHECKED(UDMXPixelMappingFixtureGroupItemComponent, bMonochromeExpose));
 	MonochromeAttribute->InvertHandle = DetailLayout->GetProperty(GET_MEMBER_NAME_CHECKED(UDMXPixelMappingFixtureGroupItemComponent, bMonochromeInvert));
 	MonochromeAttributes.Add(MonochromeAttribute);
 
+	// Hide absolute postition property handles
+	TSharedPtr<IPropertyHandle> PositionXPropertyHandle = InDetailLayout.GetProperty(UDMXPixelMappingOutputComponent::GetPositionXPropertyName());
+	InDetailLayout.HideProperty(PositionXPropertyHandle);
+	TSharedPtr<IPropertyHandle> PositionYPropertyHandle = InDetailLayout.GetProperty(UDMXPixelMappingOutputComponent::GetPositionYPropertyName());
+	InDetailLayout.HideProperty(PositionYPropertyHandle);
+	TSharedPtr<IPropertyHandle> SizeXPropertyHandle = InDetailLayout.GetProperty(UDMXPixelMappingOutputComponent::GetSizeXPropertyName());
+	InDetailLayout.HideProperty(SizeXPropertyHandle);
+	TSharedPtr<IPropertyHandle> SizeYPropertyHandle = InDetailLayout.GetProperty(UDMXPixelMappingOutputComponent::GetSizeYPropertyName());
+	InDetailLayout.HideProperty(SizeXPropertyHandle);
+
+	// Add Function and ColorMode properties at the beginning
+	TSharedPtr<IPropertyHandle> ColorModePropertyHandle = DetailLayout->GetProperty(GET_MEMBER_NAME_CHECKED(UDMXPixelMappingFixtureGroupItemComponent, ColorMode));
+	IDetailCategoryBuilder& OutputSettingsCategory = DetailLayout->EditCategory("Output Settings", FText::GetEmpty(), ECategoryPriority::Important);
+	OutputSettingsCategory.AddProperty(ColorModePropertyHandle);
+
 	// Generate all RGB Expose and Invert rows
+	for (TSharedPtr<FFunctionAttribute> Attribute : RGBAttributes)
+	{
+		InDetailLayout.HideProperty(Attribute->ExposeHandle);
+		InDetailLayout.HideProperty(Attribute->InvertHandle);
+	}
+
+	for (TSharedPtr<FFunctionAttribute> Attribute : MonochromeAttributes)
+	{
+		InDetailLayout.HideProperty(Attribute->ExposeHandle);
+		InDetailLayout.HideProperty(Attribute->InvertHandle);
+	}
+
 	OutputSettingsCategory.AddCustomRow(FText::GetEmpty())
 		.Visibility(TAttribute<EVisibility>(this, &FDMXPixelMappingDetailCustomization_FixtureGroupItem::GetRGBAttributesVisibility))
 		.NameContent()
@@ -175,8 +184,36 @@ void FDMXPixelMappingDetailCustomization_FixtureGroupItem::CustomizeDetails(IDet
 			.OnGenerateRow(this, &FDMXPixelMappingDetailCustomization_FixtureGroupItem::GenerateExposeAndInvertRow)
 		];
 
+	// Register a property type customization for the attributes as we can display the attributes of the fixture patches
+	TArray<TWeakObjectPtr<UDMXEntityFixturePatch>> FixturePatches = GetFixturePatchFromGroupItemComponents(FixtureGroupItemComponents);
+	FOnGetPropertyTypeCustomizationInstance OnGetPropertyTypeCustomizationInstance = FOnGetPropertyTypeCustomizationInstance::CreateStatic(&FDMXPixelMappingAttributeNamesDetails::MakeInstance, FixturePatches);
+	InDetailLayout.GetDetailsView()->RegisterInstancedCustomPropertyTypeLayout(FDMXAttributeName::StaticStruct()->GetFName(), OnGetPropertyTypeCustomizationInstance);
+
+	// Add RGB attributes
+	const FName RGBCategoryName = AttributeR->Handle->GetDefaultCategoryName();
+	const TAttribute<EVisibility> RGBVisibilityAttribute = TAttribute<EVisibility>::CreateLambda([this]()
+		{
+			return GetRGBAttributesVisibility();
+		});
+	InDetailLayout.EditCategory(RGBCategoryName)
+		.AddProperty(AttributeR->Handle)
+		.Visibility(RGBVisibilityAttribute);
+	InDetailLayout.EditCategory(RGBCategoryName)
+		.AddProperty(AttributeG->Handle)
+		.Visibility(RGBVisibilityAttribute);
+	InDetailLayout.EditCategory(RGBCategoryName)
+		.AddProperty(AttributeB->Handle)
+		.Visibility(RGBVisibilityAttribute);
+
+	// Add Monochrome attribute
+	InDetailLayout.EditCategory(MonochromeAttribute->Handle->GetDefaultCategoryName())
+		.AddProperty(MonochromeAttribute->Handle)
+		.Visibility(TAttribute<EVisibility>::CreateLambda([this]()
+			{
+				return GetMonochromeAttributesVisibility();
+			}));
+
 	CreateModulatorDetails(InDetailLayout);
-	CreateAttributeDetails(InDetailLayout);
 }
 
 bool FDMXPixelMappingDetailCustomization_FixtureGroupItem::CheckComponentsDMXColorMode(const EDMXColorMode DMXColorMode) const
@@ -192,56 +229,6 @@ bool FDMXPixelMappingDetailCustomization_FixtureGroupItem::CheckComponentsDMXCol
 	return false;
 }
 
-EVisibility FDMXPixelMappingDetailCustomization_FixtureGroupItem::GetRGBAttributeRowVisibilty(FFunctionAttribute* Attribute) const
-{
-	bool bIsVisible = false;
-
-	// 1. Check if current attribute is sampling now
-	FPropertyAccess::Result Result = Attribute->ExposeHandle->GetValue(bIsVisible);
-	if (Result == FPropertyAccess::Result::MultipleValues)
-	{
-		bIsVisible = true;
-	}
-
-	// 2. Check if current color mode is RGB
-	if (!CheckComponentsDMXColorMode(EDMXColorMode::CM_RGB))
-	{
-		bIsVisible = false;
-	}
-
-	return bIsVisible ? EVisibility::Visible : EVisibility::Collapsed;
-}
-
-EVisibility FDMXPixelMappingDetailCustomization_FixtureGroupItem::GetRGBAttributesVisibility() const
-{
-	return CheckComponentsDMXColorMode(EDMXColorMode::CM_RGB) ? EVisibility::Visible : EVisibility::Collapsed;
-}
-
-EVisibility FDMXPixelMappingDetailCustomization_FixtureGroupItem::GetMonochromeRowVisibilty(FFunctionAttribute* Attribute) const
-{
-	bool bIsVisible = false;
-
-	// 1. Check if current attribute is sampling now
-	FPropertyAccess::Result Result = Attribute->ExposeHandle->GetValue(bIsVisible);
-	if (Result == FPropertyAccess::Result::MultipleValues)
-	{
-		bIsVisible = true;
-	}
-
-	// 2. Check if current color mode is Monochrome
-	if (!CheckComponentsDMXColorMode(EDMXColorMode::CM_Monochrome))
-	{
-		bIsVisible = false;
-	}
-
-	return bIsVisible ? EVisibility::Visible : EVisibility::Collapsed;
-}
-
-EVisibility FDMXPixelMappingDetailCustomization_FixtureGroupItem::GetMonochromeAttributesVisibility() const
-{
-	return (GetRGBAttributesVisibility() == EVisibility::Visible) ? EVisibility::Collapsed : EVisibility::Visible;
-}
-
 TSharedRef<ITableRow> FDMXPixelMappingDetailCustomization_FixtureGroupItem::GenerateExposeAndInvertRow(TSharedPtr<FFunctionAttribute> InAttribute, const TSharedRef<STableViewBase>& OwnerTable)
 {
 	if (!InAttribute.IsValid())
@@ -292,6 +279,16 @@ TSharedRef<ITableRow> FDMXPixelMappingDetailCustomization_FixtureGroupItem::Gene
 		];
 }
 
+EVisibility FDMXPixelMappingDetailCustomization_FixtureGroupItem::GetRGBAttributesVisibility() const
+{
+	return CheckComponentsDMXColorMode(EDMXColorMode::CM_RGB) ? EVisibility::Visible : EVisibility::Collapsed;
+}
+
+EVisibility FDMXPixelMappingDetailCustomization_FixtureGroupItem::GetMonochromeAttributesVisibility() const
+{
+	return (GetRGBAttributesVisibility() == EVisibility::Visible) ? EVisibility::Collapsed : EVisibility::Visible;
+}
+
 void FDMXPixelMappingDetailCustomization_FixtureGroupItem::CreateModulatorDetails(IDetailLayoutBuilder& InDetailLayout)
 {
 	IDetailCategoryBuilder& ModualtorsCategory = InDetailLayout.EditCategory("Modulators", LOCTEXT("DMXModulatorsCategory", "Modulators"), ECategoryPriority::Important);
@@ -361,6 +358,12 @@ void FDMXPixelMappingDetailCustomization_FixtureGroupItem::CreateModulatorDetail
 					TSharedRef<IDetailsView> DetailsView = PropertyEditorModule.CreateDetailView(DetailsViewArgs);
 					DetailsView->SetObjects(ModulatorsToEdit);
 
+					TArray<TWeakObjectPtr<UDMXEntityFixturePatch>> FixturePatches = GetFixturePatchFromGroupItemComponents(FixtureGroupItemComponents);
+					FOnGetPropertyTypeCustomizationInstance OnGetPropertyTypeCustomizationInstance = FOnGetPropertyTypeCustomizationInstance::CreateStatic(&FDMXPixelMappingAttributeNamesDetails::MakeInstance, FixturePatches);
+					DetailsView->RegisterInstancedCustomPropertyTypeLayout(FDMXAttributeName::StaticStruct()->GetFName(), OnGetPropertyTypeCustomizationInstance);
+
+					PropertyEditorModule.NotifyCustomizationModuleChanged();
+
 					ModualtorsCategory.AddCustomRow(FText::GetEmpty())
 						.WholeRowContent()
 						[
@@ -384,118 +387,22 @@ void FDMXPixelMappingDetailCustomization_FixtureGroupItem::CreateModulatorDetail
 	}
 }
 
-void FDMXPixelMappingDetailCustomization_FixtureGroupItem::CreateAttributeDetails(IDetailLayoutBuilder& InDetailLayout)
+TArray<TWeakObjectPtr<UDMXEntityFixturePatch>> FDMXPixelMappingDetailCustomization_FixtureGroupItem::GetFixturePatchFromGroupItemComponents(const TArray<TWeakObjectPtr<UDMXPixelMappingFixtureGroupItemComponent>>& GroupItemComponents)
 {
-	TArray<TWeakObjectPtr<UObject>> Objects;
-	DetailLayout->GetObjectsBeingCustomized(Objects);
-
-	TArray<FName> FixtureGroupItemComponentsAttributes;
+	TArray<TWeakObjectPtr<UDMXEntityFixturePatch>> FixturePatches;
 
-	TArray<UDMXPixelMappingFixtureGroupItemComponent*> GroupItemComponents;
-	for (TWeakObjectPtr<UObject> SelectedObject : Objects)
-	{
-		UDMXPixelMappingFixtureGroupItemComponent* FixtureGroupItemComponent = Cast<UDMXPixelMappingFixtureGroupItemComponent>(SelectedObject);
-		if (FixtureGroupItemComponent)
-		{
-			GroupItemComponents.Add(FixtureGroupItemComponent);
-		}
-	}
-
-	// Gather attribute names present in Patches of all selected components
-	for (UDMXPixelMappingFixtureGroupItemComponent* GroupItemComponent : GroupItemComponents)
+	for (TWeakObjectPtr<UDMXPixelMappingFixtureGroupItemComponent> GroupItemComponent : GroupItemComponents)
 	{
-		const UDMXEntityFixturePatch* FixturePatch = GroupItemComponent->FixturePatchRef.GetFixturePatch();
+		UDMXEntityFixturePatch* FixturePatch = GroupItemComponent->FixturePatchRef.GetFixturePatch();
 		if (!FixturePatch)
 		{
 			continue;
 		}
 
-		const TArray<FDMXAttributeName> AttributeNames = FixturePatch->GetAllAttributesInActiveMode();
-		for (const FDMXAttributeName& AttributeName : AttributeNames)
-		{
-			FixtureGroupItemComponentsAttributes.AddUnique(AttributeName.Name);
-		}
+		FixturePatches.Add(TWeakObjectPtr<UDMXEntityFixturePatch>(FixturePatch));
 	}
 
-	for (UDMXPixelMappingFixtureGroupItemComponent* GroupItemComponent : GroupItemComponents)
-	{
-		FixtureGroupItemComponentsAttributes.RemoveAll([GroupItemComponent](const FName& AttributeName)
-			{
-				const UDMXEntityFixturePatch* FixturePatch = GroupItemComponent->FixturePatchRef.GetFixturePatch();
-				if (!FixturePatch)
-				{
-					return true;
-				}
-				const TArray<FDMXAttributeName> AttributeNamesOfPatch = FixturePatch->GetAllAttributesInActiveMode();
-				return !AttributeNamesOfPatch.Contains(AttributeName);
-			});
-	}
-
-	// Lambda to add a single attribute property row
-	auto CreateAttributePropertyRowLambda = [&InDetailLayout, &FixtureGroupItemComponentsAttributes](const TSharedRef<FFunctionAttribute>& Attribute, const TAttribute<EVisibility>& VisibilityAttribute)
-		{
-			InDetailLayout.HideProperty(Attribute->ExposeHandle);
-			InDetailLayout.HideProperty(Attribute->InvertHandle);
-			InDetailLayout.HideProperty(Attribute->Handle);
-
-			const TSharedRef< SDMXPixelMappingAttributeNamesComboBox> AttributeNameComboBox =
-				SNew(SDMXPixelMappingAttributeNamesComboBox, FixtureGroupItemComponentsAttributes)
-				.OnSelectionChanged_Lambda([Attribute](const FName& NewValue)
-					{
-						Attribute->SetAttributeValue(NewValue);
-					});
-
-			if (Attribute->HasMultipleAttributeValues())
-			{
-				AttributeNameComboBox->SetHasMultipleValues();
-			}
-			else if (FixtureGroupItemComponentsAttributes.Num() > 0)
-			{
-				FName InitialSelection = Attribute->GetAttributeValue();
-
-				// Set a valid attribute, or name none if none is available
-				if (!FixtureGroupItemComponentsAttributes.Contains(InitialSelection))
-				{
-					if (FixtureGroupItemComponentsAttributes.IsEmpty())
-					{
-						Attribute->SetAttributeValue(NAME_None);
-					}
-					else
-					{
-						Attribute->SetAttributeValue(FixtureGroupItemComponentsAttributes[0]);
-						InitialSelection = FixtureGroupItemComponentsAttributes[0];
-					}
-				}
-				AttributeNameComboBox->SetSelection(InitialSelection);
-			}
-
-			InDetailLayout.EditCategory("Output Settings", FText::GetEmpty(), ECategoryPriority::Important)
-				.AddCustomRow(FText::GetEmpty())
-				.Visibility(VisibilityAttribute)
-				.NameContent()
-				[
-					SNew(STextBlock)
-					.Text(Attribute->Handle->GetPropertyDisplayName())
-					.Font(FAppStyle::GetFontStyle(TEXT("PropertyWindow.NormalFont")))
-				]
-				.ValueContent()
-				[
-					AttributeNameComboBox
-				];
-		};
-
-	// Create a row for each attribute
-	for (TSharedPtr<FFunctionAttribute>& Attribute : RGBAttributes)
-	{
-		const TAttribute<EVisibility> VisibilityAttribute = TAttribute<EVisibility>(this, &FDMXPixelMappingDetailCustomization_FixtureGroupItem::GetRGBAttributesVisibility);
-		CreateAttributePropertyRowLambda(Attribute.ToSharedRef(), VisibilityAttribute);
-	}
-
-	for (TSharedPtr<FFunctionAttribute>& Attribute : MonochromeAttributes)
-	{
-		const TAttribute<EVisibility> VisibilityAttribute = TAttribute<EVisibility>(this, &FDMXPixelMappingDetailCustomization_FixtureGroupItem::GetMonochromeAttributesVisibility);
-		CreateAttributePropertyRowLambda(Attribute.ToSharedRef(), VisibilityAttribute);
-	}
+	return FixturePatches;
 }
 
 void FDMXPixelMappingDetailCustomization_FixtureGroupItem::ForceRefresh()
diff --git a/Engine/Plugins/VirtualProduction/DMX/DMXPixelMapping/Source/DMXPixelMappingEditor/Private/Customizations/DMXPixelMappingDetailCustomization_FixtureGroupItem.h b/Engine/Plugins/VirtualProduction/DMX/DMXPixelMapping/Source/DMXPixelMappingEditor/Private/Customizations/DMXPixelMappingDetailCustomization_FixtureGroupItem.h
--- a/Engine/Plugins/VirtualProduction/DMX/DMXPixelMapping/Source/DMXPixelMappingEditor/Private/Customizations/DMXPixelMappingDetailCustomization_FixtureGroupItem.h
+++ b/Engine/Plugins/VirtualProduction/DMX/DMXPixelMapping/Source/DMXPixelMappingEditor/Private/Customizations/DMXPixelMappingDetailCustomization_FixtureGroupItem.h
@@ -3,7 +3,6 @@
 #pragma once
 
 #include "CoreMinimal.h"
-#include "DMXAttribute.h"
 #include "IDetailCustomization.h"
 #include "PropertyHandle.h"
 #include "Misc/Attribute.h"
@@ -16,6 +15,7 @@ class IDetailLayoutBuilder;
 class ITableRow;
 class STableViewBase;
 class IPropertyHandle;
+class UDMXEntityFixturePatch;
 class UDMXPixelMappingFixtureGroupItemComponent;
 enum class EDMXColorMode : uint8;
 
@@ -32,7 +32,6 @@ private:
 		TSharedPtr<IPropertyHandle> Handle;
 		TSharedPtr<IPropertyHandle> ExposeHandle;
 		TSharedPtr<IPropertyHandle> InvertHandle;
-		TAttribute<EVisibility> Visibility;
 
 		/** Returns true if the Attribute Handle has multiple values */
 		bool HasMultipleAttributeValues() const;
@@ -60,24 +59,19 @@ public:
 	//~ IPropertyTypeCustomization interface end
 
 private:
-	EVisibility GetRGBAttributeRowVisibilty(FFunctionAttribute* Attribute) const;
-
 	EVisibility GetRGBAttributesVisibility() const;
 
-	EVisibility GetMonochromeRowVisibilty(FFunctionAttribute* Attribute) const;
-
 	EVisibility GetMonochromeAttributesVisibility() const;
 
 	TSharedRef<ITableRow> GenerateExposeAndInvertRow(TSharedPtr<FFunctionAttribute> InAttribute, const TSharedRef<STableViewBase>& OwnerTable);
 
-private:
 	bool CheckComponentsDMXColorMode(const EDMXColorMode DMXColorMode) const;
 
 	/** Creates Details for the Output Modulators */
 	void CreateModulatorDetails(IDetailLayoutBuilder& InDetailLayout);
 
-	/** Creates the Details for the Attributes */
-	void CreateAttributeDetails(IDetailLayoutBuilder& InDetailLayout);
+	/** Gets an array of Fixture Patches from Group Item Components */
+	TArray<TWeakObjectPtr<UDMXEntityFixturePatch>> GetFixturePatchFromGroupItemComponents(const TArray<TWeakObjectPtr<UDMXPixelMappingFixtureGroupItemComponent>>& GroupItemComponents);
 
 	/** Forces the layout to redraw */
 	void ForceRefresh();
diff --git a/Engine/Plugins/VirtualProduction/DMX/DMXPixelMapping/Source/DMXPixelMappingEditor/Private/Customizations/DMXPixelMappingDetailCustomization_Matrix.cpp b/Engine/Plugins/VirtualProduction/DMX/DMXPixelMapping/Source/DMXPixelMappingEditor/Private/Customizations/DMXPixelMappingDetailCustomization_Matrix.cpp
--- a/Engine/Plugins/VirtualProduction/DMX/DMXPixelMapping/Source/DMXPixelMappingEditor/Private/Customizations/DMXPixelMappingDetailCustomization_Matrix.cpp
+++ b/Engine/Plugins/VirtualProduction/DMX/DMXPixelMapping/Source/DMXPixelMappingEditor/Private/Customizations/DMXPixelMappingDetailCustomization_Matrix.cpp
@@ -2,6 +2,7 @@
 
 #include "Customizations/DMXPixelMappingDetailCustomization_Matrix.h"
 
+#include "DMXAttribute.h"
 #include "DMXEditorStyle.h"
 #include "DMXPixelMapping.h"
 #include "DMXPixelMappingEditorUtils.h"
@@ -9,10 +10,10 @@
 #include "Components/DMXPixelMappingMatrixComponent.h"
 #include "Components/DMXPixelMappingMatrixCellComponent.h"
 #include "Components/DMXPixelMappingRootComponent.h"
+#include "Customizations/DMXPixelMappingAttributeNamesDetails.h"
 #include "Library/DMXEntityFixturePatch.h"
 #include "Modules/ModuleManager.h"
 #include "Toolkits/DMXPixelMappingToolkit.h"
-#include "Widgets/SDMXPixelMappingAttributeNamesComboBox.h"
 
 #include "DetailWidgetRow.h"
 #include "DetailLayoutBuilder.h"
@@ -104,17 +105,7 @@ void FDMXPixelMappingDetailCustomization_Matrix::CustomizeDetails(IDetailLayoutB
 		MatrixComponents.Add(Cast<UDMXPixelMappingMatrixComponent>(SelectedObject));
 	}
 
-	// Get editing categories
-	IDetailCategoryBuilder& OutputSettingsCategory = DetailLayout->EditCategory("Output Settings", FText::GetEmpty(), ECategoryPriority::Important);
-
-	// Hide the Layout Script property (shown in its own panel, see SDMXPixelMappingLayoutView)
-	InDetailLayout.HideProperty(GET_MEMBER_NAME_CHECKED(UDMXPixelMappingMatrixComponent, LayoutScript));
-
-	// Add color mode property
-	ColorModePropertyHandle = DetailLayout->GetProperty(GET_MEMBER_NAME_CHECKED(UDMXPixelMappingMatrixComponent, ColorMode), UDMXPixelMappingMatrixComponent::StaticClass());
-	OutputSettingsCategory.AddProperty(ColorModePropertyHandle);
-
-	// Register attributes
+	// Gather attribute property handles
 	TSharedPtr<FDMXCellAttributeGroup> AttributeR = MakeShared<FDMXCellAttributeGroup>();
 	AttributeR->Handle = DetailLayout->GetProperty(GET_MEMBER_NAME_CHECKED(UDMXPixelMappingMatrixComponent, AttributeR));
 	AttributeR->ExposeHandle = DetailLayout->GetProperty(GET_MEMBER_NAME_CHECKED(UDMXPixelMappingMatrixComponent, AttributeRExpose));
@@ -134,14 +125,33 @@ void FDMXPixelMappingDetailCustomization_Matrix::CustomizeDetails(IDetailLayoutB
 	RGBAttributes.Add(AttributeG);
 	RGBAttributes.Add(AttributeB);
 
-	// Register Monochrome attribute
 	TSharedPtr<FDMXCellAttributeGroup> MonochromeAttribute = MakeShared<FDMXCellAttributeGroup>();
 	MonochromeAttribute->Handle = DetailLayout->GetProperty(GET_MEMBER_NAME_CHECKED(UDMXPixelMappingMatrixComponent, MonochromeIntensity));
 	MonochromeAttribute->ExposeHandle = DetailLayout->GetProperty(GET_MEMBER_NAME_CHECKED(UDMXPixelMappingMatrixComponent, bMonochromeExpose));
 	MonochromeAttribute->InvertHandle = DetailLayout->GetProperty(GET_MEMBER_NAME_CHECKED(UDMXPixelMappingMatrixComponent, bMonochromeInvert));
 	MonochromeAttributes.Add(MonochromeAttribute);
 
+	// Hide the Layout Script property (shown in its own panel, see SDMXPixelMappingLayoutView)
+	InDetailLayout.HideProperty(GET_MEMBER_NAME_CHECKED(UDMXPixelMappingMatrixComponent, LayoutScript));
+
+	// Add color mode property
+	ColorModePropertyHandle = DetailLayout->GetProperty(GET_MEMBER_NAME_CHECKED(UDMXPixelMappingMatrixComponent, ColorMode), UDMXPixelMappingMatrixComponent::StaticClass());
+	IDetailCategoryBuilder& OutputSettingsCategory = DetailLayout->EditCategory("Output Settings", FText::GetEmpty(), ECategoryPriority::Important);
+	OutputSettingsCategory.AddProperty(ColorModePropertyHandle);
+
 	// Generate all RGB Expose and Invert rows
+	for (TSharedPtr<FDMXCellAttributeGroup> Attribute : RGBAttributes)
+	{
+		InDetailLayout.HideProperty(Attribute->ExposeHandle);
+		InDetailLayout.HideProperty(Attribute->InvertHandle);
+	};
+
+	for (TSharedPtr<FDMXCellAttributeGroup> Attribute : MonochromeAttributes)
+	{
+		InDetailLayout.HideProperty(Attribute->ExposeHandle);
+		InDetailLayout.HideProperty(Attribute->InvertHandle);
+	};
+
 	OutputSettingsCategory.AddCustomRow(FText::GetEmpty())
 		.Visibility(TAttribute<EVisibility>(this, &FDMXPixelMappingDetailCustomization_Matrix::GetRGBAttributesVisibility))
 		.NameContent()
@@ -169,8 +179,37 @@ void FDMXPixelMappingDetailCustomization_Matrix::CustomizeDetails(IDetailLayoutB
 			.OnGenerateRow(this, &FDMXPixelMappingDetailCustomization_Matrix::GenerateExposeAndInvertRow)
 		];
 
+	// Register a property type customization for the attributes as we can display the attributes of the fixture patches
+	TArray<TWeakObjectPtr<UDMXEntityFixturePatch>> FixturePatches = GetFixturePatchFromMatrixComponents(MatrixComponents);
+	FOnGetPropertyTypeCustomizationInstance OnGetPropertyTypeCustomizationInstance = FOnGetPropertyTypeCustomizationInstance::CreateStatic(&FDMXPixelMappingAttributeNamesDetails::MakeInstance, FixturePatches);
+	InDetailLayout.GetDetailsView()->RegisterInstancedCustomPropertyTypeLayout(FDMXAttributeName::StaticStruct()->GetFName(), OnGetPropertyTypeCustomizationInstance);
+
+	// Add RGB attributes 
+	const FName RGBCategoryName = AttributeR->Handle->GetDefaultCategoryName();
+	const TAttribute<EVisibility> RGBVisibilityAttribute = TAttribute<EVisibility>::CreateLambda([this]()
+		{
+			return GetRGBAttributesVisibility();
+		});
+	InDetailLayout.EditCategory(RGBCategoryName)
+		.AddProperty(AttributeR->Handle)
+		.Visibility(RGBVisibilityAttribute);
+	InDetailLayout.EditCategory(RGBCategoryName)
+		.AddProperty(AttributeG->Handle)
+		.Visibility(RGBVisibilityAttribute);
+	InDetailLayout.EditCategory(RGBCategoryName)
+		.AddProperty(AttributeB->Handle)
+		.Visibility(RGBVisibilityAttribute);
+
+	// Add Monochrome attribute
+	InDetailLayout.EditCategory(MonochromeAttribute->Handle->GetDefaultCategoryName())
+		.AddProperty(MonochromeAttribute->Handle)
+		.Visibility(TAttribute<EVisibility>::CreateLambda([this]()
+			{
+				return GetMonochromeAttributesVisibility();
+			}));
+
+
 	CreateModulatorDetails(InDetailLayout);
-	CreateAttributeDetails(InDetailLayout);
 }
 
 bool FDMXPixelMappingDetailCustomization_Matrix::CheckComponentsDMXColorMode(const EDMXColorMode DMXColorMode) const
@@ -186,51 +225,11 @@ bool FDMXPixelMappingDetailCustomization_Matrix::CheckComponentsDMXColorMode(con
 	return false;
 }
 
-EVisibility FDMXPixelMappingDetailCustomization_Matrix::GetRGBAttributeRowVisibilty(FDMXCellAttributeGroup* Attribute) const
-{
-	bool bIsVisible = false;
-
-	// 1. Check if current attribute is sampling now
-	FPropertyAccess::Result Result = Attribute->ExposeHandle->GetValue(bIsVisible);
-	if (Result == FPropertyAccess::Result::MultipleValues)
-	{
-		bIsVisible = true;
-	}
-
-	// 2. Check if current color mode is RGB
-	if (!CheckComponentsDMXColorMode(EDMXColorMode::CM_RGB))
-	{
-		bIsVisible = false;
-	}
-
-	return bIsVisible ? EVisibility::Visible : EVisibility::Collapsed;
-}
-
 EVisibility FDMXPixelMappingDetailCustomization_Matrix::GetRGBAttributesVisibility() const
 {	
 	return CheckComponentsDMXColorMode(EDMXColorMode::CM_RGB) ? EVisibility::Visible : EVisibility::Collapsed;
 }
 
-EVisibility FDMXPixelMappingDetailCustomization_Matrix::GetMonochromeRowVisibilty(FDMXCellAttributeGroup* Attribute) const
-{
-	bool bIsVisible = false;
-
-	// 1. Check if current attribute is sampling now
-	FPropertyAccess::Result Result = Attribute->ExposeHandle->GetValue(bIsVisible);
-	if (Result == FPropertyAccess::Result::MultipleValues)
-	{
-		bIsVisible = true;
-	}
-
-	// 2. Check if current color mode is Monochrome
-	if (!CheckComponentsDMXColorMode(EDMXColorMode::CM_Monochrome))
-	{
-		bIsVisible = false;
-	}
-
-	return bIsVisible ? EVisibility::Visible : EVisibility::Collapsed;
-}
-
 EVisibility FDMXPixelMappingDetailCustomization_Matrix::GetMonochromeAttributesVisibility() const
 {
 	return (GetRGBAttributesVisibility() == EVisibility::Visible) ? EVisibility::Collapsed : EVisibility::Visible;
@@ -354,6 +353,13 @@ void FDMXPixelMappingDetailCustomization_Matrix::CreateModulatorDetails(IDetailL
 					TSharedRef<IDetailsView> DetailsView = PropertyEditorModule.CreateDetailView(DetailsViewArgs);
 					DetailsView->SetObjects(ModulatorsToEdit);
 
+					TArray<TWeakObjectPtr<UDMXEntityFixturePatch>> FixturePatches = GetFixturePatchFromMatrixComponents(MatrixComponents);
+					FOnGetPropertyTypeCustomizationInstance OnGetPropertyTypeCustomizationInstance = FOnGetPropertyTypeCustomizationInstance::CreateStatic(&FDMXPixelMappingAttributeNamesDetails::MakeInstance, FixturePatches);
+					DetailsView->RegisterInstancedCustomPropertyTypeLayout(FDMXAttributeName::StaticStruct()->GetFName(), OnGetPropertyTypeCustomizationInstance);
+
+					PropertyEditorModule.NotifyCustomizationModuleChanged();
+
+
 					ModualtorsCategory.AddCustomRow(FText::GetEmpty())
 						.WholeRowContent()
 						[
@@ -377,133 +383,21 @@ void FDMXPixelMappingDetailCustomization_Matrix::CreateModulatorDetails(IDetailL
 	}
 }
 
-void FDMXPixelMappingDetailCustomization_Matrix::CreateAttributeDetails(IDetailLayoutBuilder& InDetailLayout)
+TArray<TWeakObjectPtr<UDMXEntityFixturePatch>> FDMXPixelMappingDetailCustomization_Matrix::GetFixturePatchFromMatrixComponents(const TArray<TWeakObjectPtr<UDMXPixelMappingMatrixComponent>>& InMatrixComponents)
 {
-	TArray<TWeakObjectPtr<UObject>> Objects;
-	DetailLayout->GetObjectsBeingCustomized(Objects);
+	TArray<TWeakObjectPtr<UDMXEntityFixturePatch>> FixturePatches;
 
-	TArray<UDMXPixelMappingMatrixComponent*> StrongMatrixComponents;
-	for (TWeakObjectPtr<UObject> SelectedObject : Objects)
+	for (TWeakObjectPtr<UDMXPixelMappingMatrixComponent> GroupItemComponent : InMatrixComponents)
 	{
-		UDMXPixelMappingMatrixComponent* MatrixComponent = Cast<UDMXPixelMappingMatrixComponent>(SelectedObject);
-		if (MatrixComponent)
-		{
-			StrongMatrixComponents.Add(MatrixComponent);
-		}
-	}
-
-	// Gather attribute names present in Patches of all selected components
-	TArray<FName> FixtureMatrixAttributes;
-	for (UDMXPixelMappingMatrixComponent* MatrixComponent : StrongMatrixComponents)
-	{
-		const UDMXEntityFixturePatch* FixturePatch = MatrixComponent->FixturePatchRef.GetFixturePatch();
+		UDMXEntityFixturePatch* FixturePatch = GroupItemComponent->FixturePatchRef.GetFixturePatch();
 		if (!FixturePatch)
 		{
 			continue;
 		}
-
-		const FDMXFixtureMode* ModePtr = FixturePatch->GetActiveMode();
-		if (!ModePtr)
-		{
-			continue;
-		}
-
-		const TArray<FDMXFixtureCellAttribute> CellAttributes = ModePtr->FixtureMatrixConfig.CellAttributes;
-		for (const FDMXFixtureCellAttribute& CellAttribute : CellAttributes)
-		{
-			FixtureMatrixAttributes.AddUnique(CellAttribute.Attribute.Name);
-		}
-	}
-
-	for (UDMXPixelMappingMatrixComponent* MatrixComponent : StrongMatrixComponents)
-	{
-		FixtureMatrixAttributes.RemoveAll([MatrixComponent](const FName& AttributeName)
-			{
-				const UDMXEntityFixturePatch* FixturePatch = MatrixComponent->FixturePatchRef.GetFixturePatch();
-				if (!FixturePatch)
-				{
-					return true;
-				}
-
-				const FDMXFixtureMode* ModePtr = FixturePatch->GetActiveMode();
-				if (!ModePtr)
-				{
-					return true;
-				}
-
-				const TArray<FDMXFixtureCellAttribute> AttributeNamesOfPatch = ModePtr->FixtureMatrixConfig.CellAttributes;
-				return !AttributeNamesOfPatch.ContainsByPredicate([AttributeName](const FDMXFixtureCellAttribute& AttributeNameOfPatch)
-					{
-						return AttributeNameOfPatch.Attribute == AttributeName;
-					});
-			});
-	}
-
-	// Lambda to add a single attribute property row
-	auto CreateAttributePropertyRowLambda = [&InDetailLayout, &FixtureMatrixAttributes](const TSharedRef<FDMXCellAttributeGroup>& Attribute, const TAttribute<EVisibility>& VisibilityAttribute)
-		{
-			InDetailLayout.HideProperty(Attribute->ExposeHandle);
-			InDetailLayout.HideProperty(Attribute->InvertHandle);
-			InDetailLayout.HideProperty(Attribute->Handle);
-
-			const TSharedRef<SDMXPixelMappingAttributeNamesComboBox> AttributeNameComboBox =
-				SNew(SDMXPixelMappingAttributeNamesComboBox, FixtureMatrixAttributes)
-				.OnSelectionChanged_Lambda([Attribute](const FName& NewValue)
-					{
-						Attribute->SetAttributeValue(NewValue);
-					});
-
-			if (Attribute->HasMultipleAttributeValues())
-			{
-				AttributeNameComboBox->SetHasMultipleValues();
-			}
-			else
-			{
-				FName InitialSelection = Attribute->GetAttributeValue();
-
-				// Set a valid attribute, or name none if none is available
-				if (!FixtureMatrixAttributes.Contains(InitialSelection))
-				{
-					if (FixtureMatrixAttributes.IsEmpty())
-					{
-						Attribute->SetAttributeValue(NAME_None);
-					}
-					else
-					{
-						Attribute->SetAttributeValue(FixtureMatrixAttributes[0]);
-						InitialSelection = FixtureMatrixAttributes[0];
-					}
-				}
-				AttributeNameComboBox->SetSelection(InitialSelection);
-			}
-
-			InDetailLayout.EditCategory("Output Settings", FText::GetEmpty(), ECategoryPriority::Important)
-				.AddCustomRow(FText::GetEmpty())
-				.Visibility(VisibilityAttribute)
-				.NameContent()
-				[
-					SNew(STextBlock)
-					.Text(Attribute->Handle->GetPropertyDisplayName())
-					.Font(FAppStyle::GetFontStyle(TEXT("PropertyWindow.NormalFont")))
-				]
-				.ValueContent()
-				[
-					AttributeNameComboBox
-				];
-		};
-
-	// Create a row for each attribute
-	for (TSharedPtr<FDMXCellAttributeGroup>& Attribute : RGBAttributes)
-	{
-		const TAttribute<EVisibility> VisibilityAttribute = TAttribute<EVisibility>(this, &FDMXPixelMappingDetailCustomization_Matrix::GetRGBAttributesVisibility);
-		CreateAttributePropertyRowLambda(Attribute.ToSharedRef(), VisibilityAttribute);
+		FixturePatches.Add(TWeakObjectPtr<UDMXEntityFixturePatch>(FixturePatch));
 	}
 
-	for (TSharedPtr<FDMXCellAttributeGroup>& Attribute : MonochromeAttributes)
-	{
-		const TAttribute<EVisibility> VisibilityAttribute = TAttribute<EVisibility>(this, &FDMXPixelMappingDetailCustomization_Matrix::GetMonochromeAttributesVisibility);
-		CreateAttributePropertyRowLambda(Attribute.ToSharedRef(), VisibilityAttribute);
-	}
+	return FixturePatches;
 }
 
 void FDMXPixelMappingDetailCustomization_Matrix::ForceRefresh()
diff --git a/Engine/Plugins/VirtualProduction/DMX/DMXPixelMapping/Source/DMXPixelMappingEditor/Private/Customizations/DMXPixelMappingDetailCustomization_Matrix.h b/Engine/Plugins/VirtualProduction/DMX/DMXPixelMapping/Source/DMXPixelMappingEditor/Private/Customizations/DMXPixelMappingDetailCustomization_Matrix.h
--- a/Engine/Plugins/VirtualProduction/DMX/DMXPixelMapping/Source/DMXPixelMappingEditor/Private/Customizations/DMXPixelMappingDetailCustomization_Matrix.h
+++ b/Engine/Plugins/VirtualProduction/DMX/DMXPixelMapping/Source/DMXPixelMappingEditor/Private/Customizations/DMXPixelMappingDetailCustomization_Matrix.h
@@ -10,13 +10,14 @@
 
 class FDMXPixelMappingToolkit;
 struct FDMXCellAttributeGroup;
-class UDMXPixelMappingMatrixComponent;
 class UDMXEntityFixturePatch;
+class UDMXPixelMappingMatrixComponent;
 
 enum class EDMXColorMode : uint8;
 class ITableRow;
 class IPropertyHandle;
 class IPropertyUtilities;
+class SDMXPixelMappingAttributeNamesComboBox;
 class STableViewBase;
 template <typename ItemType> class SListView;
 
@@ -30,7 +31,6 @@ private:
 		TSharedPtr<IPropertyHandle> Handle;
 		TSharedPtr<IPropertyHandle> ExposeHandle;
 		TSharedPtr<IPropertyHandle> InvertHandle;
-		TAttribute<EVisibility> Visibility;
 
 		/** Returns true if the Attribute Handle has multiple values */
 		bool HasMultipleAttributeValues() const;
@@ -57,25 +57,20 @@ public:
 	virtual void CustomizeDetails(IDetailLayoutBuilder& DetailLayout) override;
 
 private:
-	EVisibility GetRGBAttributeRowVisibilty(FDMXCellAttributeGroup* Attribute) const;
-
 	EVisibility GetRGBAttributesVisibility() const;
 
-	EVisibility GetMonochromeRowVisibilty(FDMXCellAttributeGroup* Attribute) const;
-
 	EVisibility GetMonochromeAttributesVisibility() const;
 
 	TSharedRef<ITableRow> GenerateExposeAndInvertRow(TSharedPtr<FDMXCellAttributeGroup> InAtribute, const TSharedRef<STableViewBase>& OwnerTable);
 
-
 private:	
 	bool CheckComponentsDMXColorMode(const EDMXColorMode DMXColorMode) const;
 
 	/** Creates Details for the Output Modulators */
 	void CreateModulatorDetails(IDetailLayoutBuilder& InDetailLayout);
 
-	/** Creates the Details for the Attributes */
-	void CreateAttributeDetails(IDetailLayoutBuilder& InDetailLayout);
+	/** Gets an array of Fixture Patches from Matrix Components */
+	TArray<TWeakObjectPtr<UDMXEntityFixturePatch>> GetFixturePatchFromMatrixComponents(const TArray<TWeakObjectPtr<UDMXPixelMappingMatrixComponent>>& InMatrixComponents);
 
 	/** Forces the layout to redraw */
 	void ForceRefresh();
diff --git a/Engine/Plugins/VirtualProduction/DMX/DMXPixelMapping/Source/DMXPixelMappingEditor/Private/Widgets/SDMXPixelMappingAttributeNamesComboBox.cpp b/Engine/Plugins/VirtualProduction/DMX/DMXPixelMapping/Source/DMXPixelMappingEditor/Private/Widgets/SDMXPixelMappingAttributeNamesComboBox.cpp
deleted file mode 100644
diff --git a/Engine/Plugins/VirtualProduction/DMX/DMXPixelMapping/Source/DMXPixelMappingEditor/Private/Widgets/SDMXPixelMappingAttributeNamesComboBox.h b/Engine/Plugins/VirtualProduction/DMX/DMXPixelMapping/Source/DMXPixelMappingEditor/Private/Widgets/SDMXPixelMappingAttributeNamesComboBox.h
deleted file mode 100644
