DMX - Fix matrix attributes are not listed by [REDACTED_PERSON] attribute properties, rollback previous fix in 22237529, 22429810

diff --git a/Engine/Plugins/VirtualProduction/DMX/DMXEngine/Source/DMXEditor/Private/DMXInitializeFixtureTypeFromGDTFHelper.cpp b/Engine/Plugins/VirtualProduction/DMX/DMXEngine/Source/DMXEditor/Private/DMXInitializeFixtureTypeFromGDTFHelper.cpp
--- a/Engine/Plugins/VirtualProduction/DMX/DMXEngine/Source/DMXEditor/Private/DMXInitializeFixtureTypeFromGDTFHelper.cpp
+++ b/Engine/Plugins/VirtualProduction/DMX/DMXEngine/Source/DMXEditor/Private/DMXInitializeFixtureTypeFromGDTFHelper.cpp
@@ -545,6 +545,34 @@ void FDMXInitializeFixtureTypeFromGDTFHelper::CleanupAttributes(UDMXEntityFixtur
 				AssignedAttributeNames.Add(Function.Attribute.Name);
 			}
 		}
+
+		if (Mode.bFixtureMatrixEnabled)
+		{
+			TArray<FName> AssignedCellAttributeNames;
+			for (FDMXFixtureCellAttribute& CellAttribute : Mode.FixtureMatrixConfig.CellAttributes)
+			{
+				const TTuple<FName, TArray<FString>>* SettingsAttributeNameToKeywordPairPtr = Algo::FindByPredicate(SettingsAttributeNameToKeywordsMap, [CellAttribute](const TTuple<FName, TArray<FString>> AttributeNameToKeywordPair)
+					{
+						if (CellAttribute.Attribute.Name == AttributeNameToKeywordPair.Key ||
+							AttributeNameToKeywordPair.Value.Contains(CellAttribute.Attribute.Name.ToString()))
+						{
+							return true;
+						}
+
+						return false;
+					});
+
+				if (SettingsAttributeNameToKeywordPairPtr && !AssignedAttributeNames.Contains(SettingsAttributeNameToKeywordPairPtr->Key))
+				{
+					CellAttribute.Attribute.Name = SettingsAttributeNameToKeywordPairPtr->Key;
+					AssignedAttributeNames.Add(SettingsAttributeNameToKeywordPairPtr->Key);
+				}
+				else
+				{
+					AssignedAttributeNames.Add(CellAttribute.Attribute.Name);
+				}
+			}
+		}
 	}
 }
 
diff --git a/Engine/Plugins/VirtualProduction/DMX/DMXPixelMapping/Source/DMXPixelMappingEditor/DMXPixelMappingEditor.Build.cs b/Engine/Plugins/VirtualProduction/DMX/DMXPixelMapping/Source/DMXPixelMappingEditor/DMXPixelMappingEditor.Build.cs
--- a/Engine/Plugins/VirtualProduction/DMX/DMXPixelMapping/Source/DMXPixelMappingEditor/DMXPixelMappingEditor.Build.cs
+++ b/Engine/Plugins/VirtualProduction/DMX/DMXPixelMapping/Source/DMXPixelMappingEditor/DMXPixelMappingEditor.Build.cs
@@ -17,7 +17,6 @@ public class DMXPixelMappingEditor : ModuleRules
 			"DMXPixelMappingRuntime",
 			"DMXPixelMappingRenderer",
 			"DMXPixelMappingBlueprintGraph",
-			"DMXProtocol",
 			"DMXRuntime",
 			"EditorStyle",
 			"Engine",
diff --git a/Engine/Plugins/VirtualProduction/DMX/DMXPixelMapping/Source/DMXPixelMappingEditor/Private/Customizations/DMXPixelMappingAttributeNamesDetails.cpp b/Engine/Plugins/VirtualProduction/DMX/DMXPixelMapping/Source/DMXPixelMappingEditor/Private/Customizations/DMXPixelMappingAttributeNamesDetails.cpp
deleted file mode 100644
diff --git a/Engine/Plugins/VirtualProduction/DMX/DMXPixelMapping/Source/DMXPixelMappingEditor/Private/Customizations/DMXPixelMappingAttributeNamesDetails.h b/Engine/Plugins/VirtualProduction/DMX/DMXPixelMapping/Source/DMXPixelMappingEditor/Private/Customizations/DMXPixelMappingAttributeNamesDetails.h
deleted file mode 100644
diff --git a/Engine/Plugins/VirtualProduction/DMX/DMXPixelMapping/Source/DMXPixelMappingEditor/Private/Customizations/DMXPixelMappingDetailCustomization_FixtureGroupItem.cpp b/Engine/Plugins/VirtualProduction/DMX/DMXPixelMapping/Source/DMXPixelMappingEditor/Private/Customizations/DMXPixelMappingDetailCustomization_FixtureGroupItem.cpp
--- a/Engine/Plugins/VirtualProduction/DMX/DMXPixelMapping/Source/DMXPixelMappingEditor/Private/Customizations/DMXPixelMappingDetailCustomization_FixtureGroupItem.cpp
+++ b/Engine/Plugins/VirtualProduction/DMX/DMXPixelMapping/Source/DMXPixelMappingEditor/Private/Customizations/DMXPixelMappingDetailCustomization_FixtureGroupItem.cpp
@@ -2,14 +2,10 @@
 
 #include "Customizations/DMXPixelMappingDetailCustomization_FixtureGroupItem.h"
 
-#include "DMXAttribute.h"
 #include "DMXPixelMappingTypes.h"
 #include "Components/DMXPixelMappingFixtureGroupItemComponent.h"
-#include "Customizations/DMXPixelMappingAttributeNamesDetails.h"
-#include "Library/DMXEntityFixturePatch.h"
-#include "Library/DMXEntityReference.h"
-#include "Modulators/DMXModulator.h"
 #include "Toolkits/DMXPixelMappingToolkit.h"
+#include "Modulators/DMXModulator.h"
 
 #include "DetailLayoutBuilder.h"
 #include "DetailCategoryBuilder.h"
@@ -20,71 +16,8 @@
 #include "Layout/Visibility.h"
 #include "Modules/ModuleManager.h"
 
-
 #define LOCTEXT_NAMESPACE "FixtureGroupItem"
 
-bool FDMXPixelMappingDetailCustomization_FixtureGroupItem::FFunctionAttribute::HasMultipleAttributeValues() const
-{
-	TArray<const void*> RawData;
-	Handle->AccessRawData(RawData);
-
-	TSet<FName> AttributeNames;
-	for (const void* RawPtr : RawData)
-	{
-		if (RawPtr)
-		{
-			// The types we use with this customization must have a cast constructor to FName
-			AttributeNames.Add(reinterpret_cast<const FDMXAttributeName*>(RawPtr)->Name);
-		}
-	}
-
-	return AttributeNames.Num() > 1;
-}
-
-FName FDMXPixelMappingDetailCustomization_FixtureGroupItem::FFunctionAttribute::GetAttributeValue() const
-{
-	if (!ensureMsgf(!HasMultipleAttributeValues(), TEXT("Cannot get attribute value from handle when handle has mutliple values")))
-	{
-		return NAME_None;
-	}
-
-	TArray<const void*> RawData;
-	Handle->AccessRawData(RawData);
-
-	for (const void* RawPtr : RawData)
-	{
-		if (RawPtr)
-		{
-			// The types we use with this customization must have a cast constructor to FName
-			return reinterpret_cast<const FDMXAttributeName*>(RawPtr)->Name;
-		}
-	}
-
-	return NAME_None;
-}
-
-void FDMXPixelMappingDetailCustomization_FixtureGroupItem::FFunctionAttribute::SetAttributeValue(const FName& NewValue)
-{
-	FStructProperty* StructProperty = CastFieldChecked<FStructProperty>(Handle->GetProperty());
-
-	TArray<void*> RawData;
-	Handle->AccessRawData(RawData);
-
-	for (void* SingleRawData : RawData)
-	{
-		FDMXAttributeName* PreviousValue = reinterpret_cast<FDMXAttributeName*>(SingleRawData);
-		FDMXAttributeName NewAttributeName;
-		NewAttributeName.SetFromName(NewValue);
-
-		// Export new value to text format that can be imported later
-		FString TextValue;
-		StructProperty->Struct->ExportText(TextValue, &NewAttributeName, PreviousValue, nullptr, EPropertyPortFlags::PPF_None, nullptr);
-
-		// Set values on edited property handle from exported text
-		ensure(Handle->SetValueFromFormattedString(TextValue, EPropertyValueSetFlags::DefaultFlags) == FPropertyAccess::Result::Success);
-	}
-}
-
 void FDMXPixelMappingDetailCustomization_FixtureGroupItem::CustomizeDetails(IDetailLayoutBuilder& InDetailLayout)
 {
 	DetailLayout = &InDetailLayout;
@@ -95,12 +28,30 @@ void FDMXPixelMappingDetailCustomization_FixtureGroupItem::CustomizeDetails(IDet
 	DetailLayout->GetObjectsBeingCustomized(Objects);
 
 	FixtureGroupItemComponents.Empty();
+	
 	for (TWeakObjectPtr<UObject> SelectedObject : Objects)
 	{
 		FixtureGroupItemComponents.Add(Cast<UDMXPixelMappingFixtureGroupItemComponent>(SelectedObject));
 	}
-	
-	// Gather attribute property handles
+
+	// Get editing categories
+	IDetailCategoryBuilder& OutputSettingsCategory = DetailLayout->EditCategory("Output Settings", FText::GetEmpty(), ECategoryPriority::Important);
+
+	// Hide absolute postition property handles
+	TSharedPtr<IPropertyHandle> PositionXPropertyHandle = InDetailLayout.GetProperty(UDMXPixelMappingOutputComponent::GetPositionXPropertyName());
+	InDetailLayout.HideProperty(PositionXPropertyHandle);
+	TSharedPtr<IPropertyHandle> PositionYPropertyHandle = InDetailLayout.GetProperty(UDMXPixelMappingOutputComponent::GetPositionYPropertyName());
+	InDetailLayout.HideProperty(PositionYPropertyHandle);
+	TSharedPtr<IPropertyHandle> SizeXPropertyHandle = InDetailLayout.GetProperty(UDMXPixelMappingOutputComponent::GetSizeXPropertyName());
+	InDetailLayout.HideProperty(SizeXPropertyHandle);
+	TSharedPtr<IPropertyHandle> SizeYPropertyHandle = InDetailLayout.GetProperty(UDMXPixelMappingOutputComponent::GetSizeYPropertyName());
+	InDetailLayout.HideProperty(SizeXPropertyHandle);
+
+	// Add Function and ColorMode properties at the beginning
+	TSharedPtr<IPropertyHandle> ColorModePropertyHandle = DetailLayout->GetProperty(GET_MEMBER_NAME_CHECKED(UDMXPixelMappingFixtureGroupItemComponent, ColorMode));
+	OutputSettingsCategory.AddProperty(ColorModePropertyHandle);
+
+	// Register attributes
 	TSharedPtr<FFunctionAttribute> AttributeR = MakeShared<FFunctionAttribute>();
 	AttributeR->Handle = DetailLayout->GetProperty(GET_MEMBER_NAME_CHECKED(UDMXPixelMappingFixtureGroupItemComponent, AttributeR));
 	AttributeR->ExposeHandle = DetailLayout->GetProperty(GET_MEMBER_NAME_CHECKED(UDMXPixelMappingFixtureGroupItemComponent, AttributeRExpose));
@@ -120,47 +71,19 @@ void FDMXPixelMappingDetailCustomization_FixtureGroupItem::CustomizeDetails(IDet
 	RGBAttributes.Add(AttributeG);
 	RGBAttributes.Add(AttributeB);
 
+	// Register Monochrome attribute
 	TSharedPtr<FFunctionAttribute> MonochromeAttribute = MakeShared<FFunctionAttribute>();
 	MonochromeAttribute->Handle = DetailLayout->GetProperty(GET_MEMBER_NAME_CHECKED(UDMXPixelMappingFixtureGroupItemComponent, MonochromeIntensity));
 	MonochromeAttribute->ExposeHandle = DetailLayout->GetProperty(GET_MEMBER_NAME_CHECKED(UDMXPixelMappingFixtureGroupItemComponent, bMonochromeExpose));
 	MonochromeAttribute->InvertHandle = DetailLayout->GetProperty(GET_MEMBER_NAME_CHECKED(UDMXPixelMappingFixtureGroupItemComponent, bMonochromeInvert));
 	MonochromeAttributes.Add(MonochromeAttribute);
 
-	// Hide absolute postition property handles
-	TSharedPtr<IPropertyHandle> PositionXPropertyHandle = InDetailLayout.GetProperty(UDMXPixelMappingOutputComponent::GetPositionXPropertyName());
-	InDetailLayout.HideProperty(PositionXPropertyHandle);
-	TSharedPtr<IPropertyHandle> PositionYPropertyHandle = InDetailLayout.GetProperty(UDMXPixelMappingOutputComponent::GetPositionYPropertyName());
-	InDetailLayout.HideProperty(PositionYPropertyHandle);
-	TSharedPtr<IPropertyHandle> SizeXPropertyHandle = InDetailLayout.GetProperty(UDMXPixelMappingOutputComponent::GetSizeXPropertyName());
-	InDetailLayout.HideProperty(SizeXPropertyHandle);
-	TSharedPtr<IPropertyHandle> SizeYPropertyHandle = InDetailLayout.GetProperty(UDMXPixelMappingOutputComponent::GetSizeYPropertyName());
-	InDetailLayout.HideProperty(SizeXPropertyHandle);
-
-	// Add Function and ColorMode properties at the beginning
-	TSharedPtr<IPropertyHandle> ColorModePropertyHandle = DetailLayout->GetProperty(GET_MEMBER_NAME_CHECKED(UDMXPixelMappingFixtureGroupItemComponent, ColorMode));
-	IDetailCategoryBuilder& OutputSettingsCategory = DetailLayout->EditCategory("Output Settings", FText::GetEmpty(), ECategoryPriority::Important);
-	OutputSettingsCategory.AddProperty(ColorModePropertyHandle);
-
 	// Generate all RGB Expose and Invert rows
-	for (TSharedPtr<FFunctionAttribute> Attribute : RGBAttributes)
-	{
-		InDetailLayout.HideProperty(Attribute->ExposeHandle);
-		InDetailLayout.HideProperty(Attribute->InvertHandle);
-	}
-
-	for (TSharedPtr<FFunctionAttribute> Attribute : MonochromeAttributes)
-	{
-		InDetailLayout.HideProperty(Attribute->ExposeHandle);
-		InDetailLayout.HideProperty(Attribute->InvertHandle);
-	}
-
 	OutputSettingsCategory.AddCustomRow(FText::GetEmpty())
 		.Visibility(TAttribute<EVisibility>(this, &FDMXPixelMappingDetailCustomization_FixtureGroupItem::GetRGBAttributesVisibility))
 		.NameContent()
 		[
-			SNew(STextBlock)
-			.Text(LOCTEXT("ColorSample", "Color Sample"))
-			.Font(FAppStyle::GetFontStyle(TEXT("PropertyWindow.NormalFont")))
+			SNew(STextBlock).Text(LOCTEXT("ColorSample", "Color Sample"))
 		]
 		.ValueContent()
 		[
@@ -169,13 +92,23 @@ void FDMXPixelMappingDetailCustomization_FixtureGroupItem::CustomizeDetails(IDet
 			.OnGenerateRow(this, &FDMXPixelMappingDetailCustomization_FixtureGroupItem::GenerateExposeAndInvertRow)
 		];
 
+	// Update RGB attributes
+	for (TSharedPtr<FFunctionAttribute>& Attribute : RGBAttributes)
+	{
+		DetailLayout->HideProperty(Attribute->ExposeHandle);
+		DetailLayout->HideProperty(Attribute->InvertHandle);
+
+		OutputSettingsCategory
+			.AddProperty(Attribute->Handle)
+			.Visibility(TAttribute<EVisibility>::Create(TAttribute<EVisibility>::FGetter::CreateSP(this, &FDMXPixelMappingDetailCustomization_FixtureGroupItem::GetRGBAttributeRowVisibilty, Attribute.Get())));
+	}
+
 	// Generate all Monochrome Expose and Invert rows
 	OutputSettingsCategory.AddCustomRow(FText::GetEmpty())
 		.Visibility(TAttribute<EVisibility>(this, &FDMXPixelMappingDetailCustomization_FixtureGroupItem::GetMonochromeAttributesVisibility))
 		.NameContent()
 		[
 			SNew(STextBlock).Text(LOCTEXT("ColorSample", "Color Sample"))
-			.Font(FAppStyle::GetFontStyle(TEXT("PropertyWindow.NormalFont")))
 		]
 		.ValueContent()
 		[
@@ -184,34 +117,16 @@ void FDMXPixelMappingDetailCustomization_FixtureGroupItem::CustomizeDetails(IDet
 			.OnGenerateRow(this, &FDMXPixelMappingDetailCustomization_FixtureGroupItem::GenerateExposeAndInvertRow)
 		];
 
-	// Register a property type customization for the attributes as we can display the attributes of the fixture patches
-	TArray<TWeakObjectPtr<UDMXEntityFixturePatch>> FixturePatches = GetFixturePatchFromGroupItemComponents(FixtureGroupItemComponents);
-	FOnGetPropertyTypeCustomizationInstance OnGetPropertyTypeCustomizationInstance = FOnGetPropertyTypeCustomizationInstance::CreateStatic(&FDMXPixelMappingAttributeNamesDetails::MakeInstance, FixturePatches);
-	InDetailLayout.GetDetailsView()->RegisterInstancedCustomPropertyTypeLayout(FDMXAttributeName::StaticStruct()->GetFName(), OnGetPropertyTypeCustomizationInstance);
+	// Update Monochrome attributes
+	for (TSharedPtr<FFunctionAttribute>& Attribute : MonochromeAttributes)
+	{
+		DetailLayout->HideProperty(Attribute->ExposeHandle);
+		DetailLayout->HideProperty(Attribute->InvertHandle);
 
-	// Add RGB attributes
-	const FName RGBCategoryName = AttributeR->Handle->GetDefaultCategoryName();
-	const TAttribute<EVisibility> RGBVisibilityAttribute = TAttribute<EVisibility>::CreateLambda([this]()
-		{
-			return GetRGBAttributesVisibility();
-		});
-	InDetailLayout.EditCategory(RGBCategoryName)
-		.AddProperty(AttributeR->Handle)
-		.Visibility(RGBVisibilityAttribute);
-	InDetailLayout.EditCategory(RGBCategoryName)
-		.AddProperty(AttributeG->Handle)
-		.Visibility(RGBVisibilityAttribute);
-	InDetailLayout.EditCategory(RGBCategoryName)
-		.AddProperty(AttributeB->Handle)
-		.Visibility(RGBVisibilityAttribute);
-
-	// Add Monochrome attribute
-	InDetailLayout.EditCategory(MonochromeAttribute->Handle->GetDefaultCategoryName())
-		.AddProperty(MonochromeAttribute->Handle)
-		.Visibility(TAttribute<EVisibility>::CreateLambda([this]()
-			{
-				return GetMonochromeAttributesVisibility();
-			}));
+		OutputSettingsCategory
+			.AddProperty(Attribute->Handle)
+			.Visibility(TAttribute<EVisibility>::Create(TAttribute<EVisibility>::FGetter::CreateSP(this, &FDMXPixelMappingDetailCustomization_FixtureGroupItem::GetMonochromeRowVisibilty, Attribute.Get())));
+	}
 
 	CreateModulatorDetails(InDetailLayout);
 }
@@ -229,6 +144,56 @@ bool FDMXPixelMappingDetailCustomization_FixtureGroupItem::CheckComponentsDMXCol
 	return false;
 }
 
+EVisibility FDMXPixelMappingDetailCustomization_FixtureGroupItem::GetRGBAttributeRowVisibilty(FFunctionAttribute* Attribute) const
+{
+	bool bIsVisible = false;
+
+	// 1. Check if current attribute is sampling now
+	FPropertyAccess::Result Result = Attribute->ExposeHandle->GetValue(bIsVisible);
+	if (Result == FPropertyAccess::Result::MultipleValues)
+	{
+		bIsVisible = true;
+	}
+
+	// 2. Check if current color mode is RGB
+	if (!CheckComponentsDMXColorMode(EDMXColorMode::CM_RGB))
+	{
+		bIsVisible = false;
+	}
+
+	return bIsVisible ? EVisibility::Visible : EVisibility::Collapsed;
+}
+
+EVisibility FDMXPixelMappingDetailCustomization_FixtureGroupItem::GetRGBAttributesVisibility() const
+{
+	return CheckComponentsDMXColorMode(EDMXColorMode::CM_RGB) ? EVisibility::Visible : EVisibility::Collapsed;
+}
+
+EVisibility FDMXPixelMappingDetailCustomization_FixtureGroupItem::GetMonochromeRowVisibilty(FFunctionAttribute* Attribute) const
+{
+	bool bIsVisible = false;
+
+	// 1. Check if current attribute is sampling now
+	FPropertyAccess::Result Result = Attribute->ExposeHandle->GetValue(bIsVisible);
+	if (Result == FPropertyAccess::Result::MultipleValues)
+	{
+		bIsVisible = true;
+	}
+
+	// 2. Check if current color mode is Monochrome
+	if (!CheckComponentsDMXColorMode(EDMXColorMode::CM_Monochrome))
+	{
+		bIsVisible = false;
+	}
+
+	return bIsVisible ? EVisibility::Visible : EVisibility::Collapsed;
+}
+
+EVisibility FDMXPixelMappingDetailCustomization_FixtureGroupItem::GetMonochromeAttributesVisibility() const
+{
+	return (GetRGBAttributesVisibility() == EVisibility::Visible) ? EVisibility::Collapsed : EVisibility::Visible;
+}
+
 TSharedRef<ITableRow> FDMXPixelMappingDetailCustomization_FixtureGroupItem::GenerateExposeAndInvertRow(TSharedPtr<FFunctionAttribute> InAttribute, const TSharedRef<STableViewBase>& OwnerTable)
 {
 	if (!InAttribute.IsValid())
@@ -279,16 +244,6 @@ TSharedRef<ITableRow> FDMXPixelMappingDetailCustomization_FixtureGroupItem::Gene
 		];
 }
 
-EVisibility FDMXPixelMappingDetailCustomization_FixtureGroupItem::GetRGBAttributesVisibility() const
-{
-	return CheckComponentsDMXColorMode(EDMXColorMode::CM_RGB) ? EVisibility::Visible : EVisibility::Collapsed;
-}
-
-EVisibility FDMXPixelMappingDetailCustomization_FixtureGroupItem::GetMonochromeAttributesVisibility() const
-{
-	return (GetRGBAttributesVisibility() == EVisibility::Visible) ? EVisibility::Collapsed : EVisibility::Visible;
-}
-
 void FDMXPixelMappingDetailCustomization_FixtureGroupItem::CreateModulatorDetails(IDetailLayoutBuilder& InDetailLayout)
 {
 	IDetailCategoryBuilder& ModualtorsCategory = InDetailLayout.EditCategory("Modulators", LOCTEXT("DMXModulatorsCategory", "Modulators"), ECategoryPriority::Important);
@@ -358,12 +313,6 @@ void FDMXPixelMappingDetailCustomization_FixtureGroupItem::CreateModulatorDetail
 					TSharedRef<IDetailsView> DetailsView = PropertyEditorModule.CreateDetailView(DetailsViewArgs);
 					DetailsView->SetObjects(ModulatorsToEdit);
 
-					TArray<TWeakObjectPtr<UDMXEntityFixturePatch>> FixturePatches = GetFixturePatchFromGroupItemComponents(FixtureGroupItemComponents);
-					FOnGetPropertyTypeCustomizationInstance OnGetPropertyTypeCustomizationInstance = FOnGetPropertyTypeCustomizationInstance::CreateStatic(&FDMXPixelMappingAttributeNamesDetails::MakeInstance, FixturePatches);
-					DetailsView->RegisterInstancedCustomPropertyTypeLayout(FDMXAttributeName::StaticStruct()->GetFName(), OnGetPropertyTypeCustomizationInstance);
-
-					PropertyEditorModule.NotifyCustomizationModuleChanged();
-
 					ModualtorsCategory.AddCustomRow(FText::GetEmpty())
 						.WholeRowContent()
 						[
@@ -387,24 +336,6 @@ void FDMXPixelMappingDetailCustomization_FixtureGroupItem::CreateModulatorDetail
 	}
 }
 
-TArray<TWeakObjectPtr<UDMXEntityFixturePatch>> FDMXPixelMappingDetailCustomization_FixtureGroupItem::GetFixturePatchFromGroupItemComponents(const TArray<TWeakObjectPtr<UDMXPixelMappingFixtureGroupItemComponent>>& GroupItemComponents)
-{
-	TArray<TWeakObjectPtr<UDMXEntityFixturePatch>> FixturePatches;
-
-	for (TWeakObjectPtr<UDMXPixelMappingFixtureGroupItemComponent> GroupItemComponent : GroupItemComponents)
-	{
-		UDMXEntityFixturePatch* FixturePatch = GroupItemComponent->FixturePatchRef.GetFixturePatch();
-		if (!FixturePatch)
-		{
-			continue;
-		}
-
-		FixturePatches.Add(TWeakObjectPtr<UDMXEntityFixturePatch>(FixturePatch));
-	}
-
-	return FixturePatches;
-}
-
 void FDMXPixelMappingDetailCustomization_FixtureGroupItem::ForceRefresh()
 {
 	PropertyUtilities->ForceRefresh();
diff --git a/Engine/Plugins/VirtualProduction/DMX/DMXPixelMapping/Source/DMXPixelMappingEditor/Private/Customizations/DMXPixelMappingDetailCustomization_FixtureGroupItem.h b/Engine/Plugins/VirtualProduction/DMX/DMXPixelMapping/Source/DMXPixelMappingEditor/Private/Customizations/DMXPixelMappingDetailCustomization_FixtureGroupItem.h
--- a/Engine/Plugins/VirtualProduction/DMX/DMXPixelMapping/Source/DMXPixelMappingEditor/Private/Customizations/DMXPixelMappingDetailCustomization_FixtureGroupItem.h
+++ b/Engine/Plugins/VirtualProduction/DMX/DMXPixelMapping/Source/DMXPixelMappingEditor/Private/Customizations/DMXPixelMappingDetailCustomization_FixtureGroupItem.h
@@ -4,7 +4,6 @@
 
 #include "CoreMinimal.h"
 #include "IDetailCustomization.h"
-#include "PropertyHandle.h"
 #include "Misc/Attribute.h"
 #include "UObject/WeakObjectPtr.h"
 #include "IPropertyTypeCustomization.h"
@@ -15,7 +14,6 @@ class IDetailLayoutBuilder;
 class ITableRow;
 class STableViewBase;
 class IPropertyHandle;
-class UDMXEntityFixturePatch;
 class UDMXPixelMappingFixtureGroupItemComponent;
 enum class EDMXColorMode : uint8;
 
@@ -32,17 +30,10 @@ private:
 		TSharedPtr<IPropertyHandle> Handle;
 		TSharedPtr<IPropertyHandle> ExposeHandle;
 		TSharedPtr<IPropertyHandle> InvertHandle;
-
-		/** Returns true if the Attribute Handle has multiple values */
-		bool HasMultipleAttributeValues() const;
-
-		/** Gets the value of the Attribute Handle. Should not be called if HasMultipleAttributeValues returns true */
-		FName GetAttributeValue() const;
-
-		/** Sets the value of the Attribute Handle */
-		void SetAttributeValue(const FName& NewValue);
+		TAttribute<EVisibility> Visibility;
 	};
 
+
 public:
 	static TSharedRef<IDetailCustomization> MakeInstance(const TWeakPtr<FDMXPixelMappingToolkit> InToolkitWeakPtr)
 	{
@@ -59,20 +50,23 @@ public:
 	//~ IPropertyTypeCustomization interface end
 
 private:
+
+	EVisibility GetRGBAttributeRowVisibilty(FFunctionAttribute* Attribute) const;
+
 	EVisibility GetRGBAttributesVisibility() const;
 
+	EVisibility GetMonochromeRowVisibilty(FFunctionAttribute* Attribute) const;
+
 	EVisibility GetMonochromeAttributesVisibility() const;
 
 	TSharedRef<ITableRow> GenerateExposeAndInvertRow(TSharedPtr<FFunctionAttribute> InAttribute, const TSharedRef<STableViewBase>& OwnerTable);
 
+private:
 	bool CheckComponentsDMXColorMode(const EDMXColorMode DMXColorMode) const;
 
 	/** Creates Details for the Output Modulators */
 	void CreateModulatorDetails(IDetailLayoutBuilder& InDetailLayout);
 
-	/** Gets an array of Fixture Patches from Group Item Components */
-	TArray<TWeakObjectPtr<UDMXEntityFixturePatch>> GetFixturePatchFromGroupItemComponents(const TArray<TWeakObjectPtr<UDMXPixelMappingFixtureGroupItemComponent>>& GroupItemComponents);
-
 	/** Forces the layout to redraw */
 	void ForceRefresh();
 
diff --git a/Engine/Plugins/VirtualProduction/DMX/DMXPixelMapping/Source/DMXPixelMappingEditor/Private/Customizations/DMXPixelMappingDetailCustomization_Matrix.cpp b/Engine/Plugins/VirtualProduction/DMX/DMXPixelMapping/Source/DMXPixelMappingEditor/Private/Customizations/DMXPixelMappingDetailCustomization_Matrix.cpp
--- a/Engine/Plugins/VirtualProduction/DMX/DMXPixelMapping/Source/DMXPixelMappingEditor/Private/Customizations/DMXPixelMappingDetailCustomization_Matrix.cpp
+++ b/Engine/Plugins/VirtualProduction/DMX/DMXPixelMapping/Source/DMXPixelMappingEditor/Private/Customizations/DMXPixelMappingDetailCustomization_Matrix.cpp
@@ -2,7 +2,6 @@
 
 #include "Customizations/DMXPixelMappingDetailCustomization_Matrix.h"
 
-#include "DMXAttribute.h"
 #include "DMXEditorStyle.h"
 #include "DMXPixelMapping.h"
 #include "DMXPixelMappingEditorUtils.h"
@@ -10,7 +9,6 @@
 #include "Components/DMXPixelMappingMatrixComponent.h"
 #include "Components/DMXPixelMappingMatrixCellComponent.h"
 #include "Components/DMXPixelMappingRootComponent.h"
-#include "Customizations/DMXPixelMappingAttributeNamesDetails.h"
 #include "Library/DMXEntityFixturePatch.h"
 #include "Modules/ModuleManager.h"
 #include "Toolkits/DMXPixelMappingToolkit.h"
@@ -27,68 +25,6 @@
 
 #define LOCTEXT_NAMESPACE "DMXPixelMappingDetailCustomization_Matrix"
 
-bool FDMXPixelMappingDetailCustomization_Matrix::FDMXCellAttributeGroup::HasMultipleAttributeValues() const
-{
-	TArray<const void*> RawData;
-	Handle->AccessRawData(RawData);
-
-	TSet<FName> AttributeNames;
-	for (const void* RawPtr : RawData)
-	{
-		if (RawPtr)
-		{
-			// The types we use with this customization must have a cast constructor to FName
-			AttributeNames.Add(reinterpret_cast<const FDMXAttributeName*>(RawPtr)->Name);
-		}
-	}
-
-	return AttributeNames.Num() > 1;
-}
-
-FName FDMXPixelMappingDetailCustomization_Matrix::FDMXCellAttributeGroup::GetAttributeValue() const
-{
-	if (!ensureMsgf(!HasMultipleAttributeValues(), TEXT("Cannot get attribute value from handle when handle has mutliple values")))
-	{
-		return NAME_None;
-	}
-
-	TArray<const void*> RawData;
-	Handle->AccessRawData(RawData);
-
-	for (const void* RawPtr : RawData)
-	{
-		if (RawPtr)
-		{
-			// The types we use with this customization must have a cast constructor to FName
-			return reinterpret_cast<const FDMXAttributeName*>(RawPtr)->Name;
-		}
-	}
-
-	return NAME_None;
-}
-
-void FDMXPixelMappingDetailCustomization_Matrix::FDMXCellAttributeGroup::SetAttributeValue(const FName& NewValue)
-{
-	FStructProperty* StructProperty = CastFieldChecked<FStructProperty>(Handle->GetProperty());
-
-	TArray<void*> RawData;
-	Handle->AccessRawData(RawData);
-
-	for (void* SingleRawData : RawData)
-	{
-		FDMXAttributeName* PreviousValue = reinterpret_cast<FDMXAttributeName*>(SingleRawData);
-		FDMXAttributeName NewAttributeName;
-		NewAttributeName.SetFromName(NewValue);
-
-		// Export new value to text format that can be imported later
-		FString TextValue;
-		StructProperty->Struct->ExportText(TextValue, &NewAttributeName, PreviousValue, nullptr, EPropertyPortFlags::PPF_None, nullptr);
-
-		// Set values on edited property handle from exported text
-		ensure(Handle->SetValueFromFormattedString(TextValue, EPropertyValueSetFlags::DefaultFlags) == FPropertyAccess::Result::Success);
-	}
-}
-
 void FDMXPixelMappingDetailCustomization_Matrix::CustomizeDetails(IDetailLayoutBuilder& InDetailLayout)
 {
 	DetailLayout = &InDetailLayout;
@@ -105,7 +41,17 @@ void FDMXPixelMappingDetailCustomization_Matrix::CustomizeDetails(IDetailLayoutB
 		MatrixComponents.Add(Cast<UDMXPixelMappingMatrixComponent>(SelectedObject));
 	}
 
-	// Gather attribute property handles
+	// Get editing categories
+	IDetailCategoryBuilder& OutputSettingsCategory = DetailLayout->EditCategory("Output Settings", FText::GetEmpty(), ECategoryPriority::Important);
+
+	// Hide the Layout Script property (shown in its own panel, see SDMXPixelMappingLayoutView)
+	InDetailLayout.HideProperty(GET_MEMBER_NAME_CHECKED(UDMXPixelMappingMatrixComponent, LayoutScript));
+
+	// Add color mode property
+	ColorModePropertyHandle = DetailLayout->GetProperty(GET_MEMBER_NAME_CHECKED(UDMXPixelMappingMatrixComponent, ColorMode), UDMXPixelMappingMatrixComponent::StaticClass());
+	OutputSettingsCategory.AddProperty(ColorModePropertyHandle);
+
+	// Register attributes
 	TSharedPtr<FDMXCellAttributeGroup> AttributeR = MakeShared<FDMXCellAttributeGroup>();
 	AttributeR->Handle = DetailLayout->GetProperty(GET_MEMBER_NAME_CHECKED(UDMXPixelMappingMatrixComponent, AttributeR));
 	AttributeR->ExposeHandle = DetailLayout->GetProperty(GET_MEMBER_NAME_CHECKED(UDMXPixelMappingMatrixComponent, AttributeRExpose));
@@ -125,33 +71,14 @@ void FDMXPixelMappingDetailCustomization_Matrix::CustomizeDetails(IDetailLayoutB
 	RGBAttributes.Add(AttributeG);
 	RGBAttributes.Add(AttributeB);
 
+	// Register Monochrome attribute
 	TSharedPtr<FDMXCellAttributeGroup> MonochromeAttribute = MakeShared<FDMXCellAttributeGroup>();
 	MonochromeAttribute->Handle = DetailLayout->GetProperty(GET_MEMBER_NAME_CHECKED(UDMXPixelMappingMatrixComponent, MonochromeIntensity));
 	MonochromeAttribute->ExposeHandle = DetailLayout->GetProperty(GET_MEMBER_NAME_CHECKED(UDMXPixelMappingMatrixComponent, bMonochromeExpose));
 	MonochromeAttribute->InvertHandle = DetailLayout->GetProperty(GET_MEMBER_NAME_CHECKED(UDMXPixelMappingMatrixComponent, bMonochromeInvert));
 	MonochromeAttributes.Add(MonochromeAttribute);
 
-	// Hide the Layout Script property (shown in its own panel, see SDMXPixelMappingLayoutView)
-	InDetailLayout.HideProperty(GET_MEMBER_NAME_CHECKED(UDMXPixelMappingMatrixComponent, LayoutScript));
-
-	// Add color mode property
-	ColorModePropertyHandle = DetailLayout->GetProperty(GET_MEMBER_NAME_CHECKED(UDMXPixelMappingMatrixComponent, ColorMode), UDMXPixelMappingMatrixComponent::StaticClass());
-	IDetailCategoryBuilder& OutputSettingsCategory = DetailLayout->EditCategory("Output Settings", FText::GetEmpty(), ECategoryPriority::Important);
-	OutputSettingsCategory.AddProperty(ColorModePropertyHandle);
-
 	// Generate all RGB Expose and Invert rows
-	for (TSharedPtr<FDMXCellAttributeGroup> Attribute : RGBAttributes)
-	{
-		InDetailLayout.HideProperty(Attribute->ExposeHandle);
-		InDetailLayout.HideProperty(Attribute->InvertHandle);
-	};
-
-	for (TSharedPtr<FDMXCellAttributeGroup> Attribute : MonochromeAttributes)
-	{
-		InDetailLayout.HideProperty(Attribute->ExposeHandle);
-		InDetailLayout.HideProperty(Attribute->InvertHandle);
-	};
-
 	OutputSettingsCategory.AddCustomRow(FText::GetEmpty())
 		.Visibility(TAttribute<EVisibility>(this, &FDMXPixelMappingDetailCustomization_Matrix::GetRGBAttributesVisibility))
 		.NameContent()
@@ -165,6 +92,18 @@ void FDMXPixelMappingDetailCustomization_Matrix::CustomizeDetails(IDetailLayoutB
 			.OnGenerateRow(this, &FDMXPixelMappingDetailCustomization_Matrix::GenerateExposeAndInvertRow)
 		];
 
+	// Update RGB attributes
+	for (TSharedPtr<FDMXCellAttributeGroup>& Attribute : RGBAttributes)
+	{
+		DetailLayout->HideProperty(Attribute->ExposeHandle);
+		DetailLayout->HideProperty(Attribute->InvertHandle);
+
+		OutputSettingsCategory
+			.AddProperty(Attribute->Handle)
+			.Visibility(TAttribute<EVisibility>::Create(TAttribute<EVisibility>::FGetter::CreateSP(this, &FDMXPixelMappingDetailCustomization_Matrix::GetRGBAttributeRowVisibilty, Attribute.Get())));
+	}
+
+
 	// Generate all Monochrome Expose and Invert rows
 	OutputSettingsCategory.AddCustomRow(FText::GetEmpty())
 		.Visibility(TAttribute<EVisibility>(this, &FDMXPixelMappingDetailCustomization_Matrix::GetMonochromeAttributesVisibility))
@@ -179,36 +118,17 @@ void FDMXPixelMappingDetailCustomization_Matrix::CustomizeDetails(IDetailLayoutB
 			.OnGenerateRow(this, &FDMXPixelMappingDetailCustomization_Matrix::GenerateExposeAndInvertRow)
 		];
 
-	// Register a property type customization for the attributes as we can display the attributes of the fixture patches
-	TArray<TWeakObjectPtr<UDMXEntityFixturePatch>> FixturePatches = GetFixturePatchFromMatrixComponents(MatrixComponents);
-	FOnGetPropertyTypeCustomizationInstance OnGetPropertyTypeCustomizationInstance = FOnGetPropertyTypeCustomizationInstance::CreateStatic(&FDMXPixelMappingAttributeNamesDetails::MakeInstance, FixturePatches);
-	InDetailLayout.GetDetailsView()->RegisterInstancedCustomPropertyTypeLayout(FDMXAttributeName::StaticStruct()->GetFName(), OnGetPropertyTypeCustomizationInstance);
-
-	// Add RGB attributes 
-	const FName RGBCategoryName = AttributeR->Handle->GetDefaultCategoryName();
-	const TAttribute<EVisibility> RGBVisibilityAttribute = TAttribute<EVisibility>::CreateLambda([this]()
-		{
-			return GetRGBAttributesVisibility();
-		});
-	InDetailLayout.EditCategory(RGBCategoryName)
-		.AddProperty(AttributeR->Handle)
-		.Visibility(RGBVisibilityAttribute);
-	InDetailLayout.EditCategory(RGBCategoryName)
-		.AddProperty(AttributeG->Handle)
-		.Visibility(RGBVisibilityAttribute);
-	InDetailLayout.EditCategory(RGBCategoryName)
-		.AddProperty(AttributeB->Handle)
-		.Visibility(RGBVisibilityAttribute);
-
-	// Add Monochrome attribute
-	InDetailLayout.EditCategory(MonochromeAttribute->Handle->GetDefaultCategoryName())
-		.AddProperty(MonochromeAttribute->Handle)
-		.Visibility(TAttribute<EVisibility>::CreateLambda([this]()
-			{
-				return GetMonochromeAttributesVisibility();
-			}));
-
+	// Update Monochrome attributes
+	for (TSharedPtr<FDMXCellAttributeGroup>& Attribute : MonochromeAttributes)
+	{
+		DetailLayout->HideProperty(Attribute->ExposeHandle);
+		DetailLayout->HideProperty(Attribute->InvertHandle);
 
+		OutputSettingsCategory
+			.AddProperty(Attribute->Handle)
+			.Visibility(TAttribute<EVisibility>::Create(TAttribute<EVisibility>::FGetter::CreateSP(this, &FDMXPixelMappingDetailCustomization_Matrix::GetMonochromeRowVisibilty, Attribute.Get())));
+	}
+	
 	CreateModulatorDetails(InDetailLayout);
 }
 
@@ -225,11 +145,51 @@ bool FDMXPixelMappingDetailCustomization_Matrix::CheckComponentsDMXColorMode(con
 	return false;
 }
 
+EVisibility FDMXPixelMappingDetailCustomization_Matrix::GetRGBAttributeRowVisibilty(FDMXCellAttributeGroup* Attribute) const
+{
+	bool bIsVisible = false;
+
+	// 1. Check if current attribute is sampling now
+	FPropertyAccess::Result Result = Attribute->ExposeHandle->GetValue(bIsVisible);
+	if (Result == FPropertyAccess::Result::MultipleValues)
+	{
+		bIsVisible = true;
+	}
+
+	// 2. Check if current color mode is RGB
+	if (!CheckComponentsDMXColorMode(EDMXColorMode::CM_RGB))
+	{
+		bIsVisible = false;
+	}
+
+	return bIsVisible ? EVisibility::Visible : EVisibility::Collapsed;
+}
+
 EVisibility FDMXPixelMappingDetailCustomization_Matrix::GetRGBAttributesVisibility() const
 {	
 	return CheckComponentsDMXColorMode(EDMXColorMode::CM_RGB) ? EVisibility::Visible : EVisibility::Collapsed;
 }
 
+EVisibility FDMXPixelMappingDetailCustomization_Matrix::GetMonochromeRowVisibilty(FDMXCellAttributeGroup* Attribute) const
+{
+	bool bIsVisible = false;
+
+	// 1. Check if current attribute is sampling now
+	FPropertyAccess::Result Result = Attribute->ExposeHandle->GetValue(bIsVisible);
+	if (Result == FPropertyAccess::Result::MultipleValues)
+	{
+		bIsVisible = true;
+	}
+
+	// 2. Check if current color mode is Monochrome
+	if (!CheckComponentsDMXColorMode(EDMXColorMode::CM_Monochrome))
+	{
+		bIsVisible = false;
+	}
+
+	return bIsVisible ? EVisibility::Visible : EVisibility::Collapsed;
+}
+
 EVisibility FDMXPixelMappingDetailCustomization_Matrix::GetMonochromeAttributesVisibility() const
 {
 	return (GetRGBAttributesVisibility() == EVisibility::Visible) ? EVisibility::Collapsed : EVisibility::Visible;
@@ -353,13 +313,6 @@ void FDMXPixelMappingDetailCustomization_Matrix::CreateModulatorDetails(IDetailL
 					TSharedRef<IDetailsView> DetailsView = PropertyEditorModule.CreateDetailView(DetailsViewArgs);
 					DetailsView->SetObjects(ModulatorsToEdit);
 
-					TArray<TWeakObjectPtr<UDMXEntityFixturePatch>> FixturePatches = GetFixturePatchFromMatrixComponents(MatrixComponents);
-					FOnGetPropertyTypeCustomizationInstance OnGetPropertyTypeCustomizationInstance = FOnGetPropertyTypeCustomizationInstance::CreateStatic(&FDMXPixelMappingAttributeNamesDetails::MakeInstance, FixturePatches);
-					DetailsView->RegisterInstancedCustomPropertyTypeLayout(FDMXAttributeName::StaticStruct()->GetFName(), OnGetPropertyTypeCustomizationInstance);
-
-					PropertyEditorModule.NotifyCustomizationModuleChanged();
-
-
 					ModualtorsCategory.AddCustomRow(FText::GetEmpty())
 						.WholeRowContent()
 						[
@@ -383,23 +336,6 @@ void FDMXPixelMappingDetailCustomization_Matrix::CreateModulatorDetails(IDetailL
 	}
 }
 
-TArray<TWeakObjectPtr<UDMXEntityFixturePatch>> FDMXPixelMappingDetailCustomization_Matrix::GetFixturePatchFromMatrixComponents(const TArray<TWeakObjectPtr<UDMXPixelMappingMatrixComponent>>& InMatrixComponents)
-{
-	TArray<TWeakObjectPtr<UDMXEntityFixturePatch>> FixturePatches;
-
-	for (TWeakObjectPtr<UDMXPixelMappingMatrixComponent> GroupItemComponent : InMatrixComponents)
-	{
-		UDMXEntityFixturePatch* FixturePatch = GroupItemComponent->FixturePatchRef.GetFixturePatch();
-		if (!FixturePatch)
-		{
-			continue;
-		}
-		FixturePatches.Add(TWeakObjectPtr<UDMXEntityFixturePatch>(FixturePatch));
-	}
-
-	return FixturePatches;
-}
-
 void FDMXPixelMappingDetailCustomization_Matrix::ForceRefresh()
 {
 	PropertyUtilities->ForceRefresh();
diff --git a/Engine/Plugins/VirtualProduction/DMX/DMXPixelMapping/Source/DMXPixelMappingEditor/Private/Customizations/DMXPixelMappingDetailCustomization_Matrix.h b/Engine/Plugins/VirtualProduction/DMX/DMXPixelMapping/Source/DMXPixelMappingEditor/Private/Customizations/DMXPixelMappingDetailCustomization_Matrix.h
--- a/Engine/Plugins/VirtualProduction/DMX/DMXPixelMapping/Source/DMXPixelMappingEditor/Private/Customizations/DMXPixelMappingDetailCustomization_Matrix.h
+++ b/Engine/Plugins/VirtualProduction/DMX/DMXPixelMapping/Source/DMXPixelMappingEditor/Private/Customizations/DMXPixelMappingDetailCustomization_Matrix.h
@@ -10,14 +10,13 @@
 
 class FDMXPixelMappingToolkit;
 struct FDMXCellAttributeGroup;
-class UDMXEntityFixturePatch;
 class UDMXPixelMappingMatrixComponent;
+class UDMXEntityFixturePatch;
 
 enum class EDMXColorMode : uint8;
 class ITableRow;
 class IPropertyHandle;
 class IPropertyUtilities;
-class SDMXPixelMappingAttributeNamesComboBox;
 class STableViewBase;
 template <typename ItemType> class SListView;
 
@@ -31,15 +30,7 @@ private:
 		TSharedPtr<IPropertyHandle> Handle;
 		TSharedPtr<IPropertyHandle> ExposeHandle;
 		TSharedPtr<IPropertyHandle> InvertHandle;
-
-		/** Returns true if the Attribute Handle has multiple values */
-		bool HasMultipleAttributeValues() const;
-
-		/** Gets the value of the Attribute Handle. Should not be called if HasMultipleAttributeValues returns true */
-		FName GetAttributeValue() const;
-
-		/** Sets the value of the Attribute Handle */
-		void SetAttributeValue(const FName& NewValue);
+		TAttribute<EVisibility> Visibility;
 	};
 
 public:
@@ -57,21 +48,23 @@ public:
 	virtual void CustomizeDetails(IDetailLayoutBuilder& DetailLayout) override;
 
 private:
+	EVisibility GetRGBAttributeRowVisibilty(FDMXCellAttributeGroup* Attribute) const;
+
 	EVisibility GetRGBAttributesVisibility() const;
 
+	EVisibility GetMonochromeRowVisibilty(FDMXCellAttributeGroup* Attribute) const;
+
 	EVisibility GetMonochromeAttributesVisibility() const;
 
 	TSharedRef<ITableRow> GenerateExposeAndInvertRow(TSharedPtr<FDMXCellAttributeGroup> InAtribute, const TSharedRef<STableViewBase>& OwnerTable);
 
+
 private:	
 	bool CheckComponentsDMXColorMode(const EDMXColorMode DMXColorMode) const;
 
 	/** Creates Details for the Output Modulators */
 	void CreateModulatorDetails(IDetailLayoutBuilder& InDetailLayout);
 
-	/** Gets an array of Fixture Patches from Matrix Components */
-	TArray<TWeakObjectPtr<UDMXEntityFixturePatch>> GetFixturePatchFromMatrixComponents(const TArray<TWeakObjectPtr<UDMXPixelMappingMatrixComponent>>& InMatrixComponents);
-
 	/** Forces the layout to redraw */
 	void ForceRefresh();
 
