DMX - Fix Undo in PixelMapping doesn't work for common patch settings

diff --git a/Engine/Plugins/VirtualProduction/DMX/DMXEngine/Source/DMXEditor/Private/Customizations/DMXAttributeNameCustomization.cpp b/Engine/Plugins/VirtualProduction/DMX/DMXEngine/Source/DMXEditor/Private/Customizations/DMXAttributeNameCustomization.cpp
--- a/Engine/Plugins/VirtualProduction/DMX/DMXEngine/Source/DMXEditor/Private/Customizations/DMXAttributeNameCustomization.cpp
+++ b/Engine/Plugins/VirtualProduction/DMX/DMXEngine/Source/DMXEditor/Private/Customizations/DMXAttributeNameCustomization.cpp
@@ -6,6 +6,7 @@
 #include "DMXProtocolSettings.h"
 
 #include "DetailWidgetRow.h"
+#include "Editor.h"
 #include "IPropertyTypeCustomization.h"
 #include "IPropertyUtilities.h"
 #include "ScopedTransaction.h"
@@ -60,7 +61,9 @@ void FDMXAttributeNameCustomization::CustomizeHeader(TSharedRef<IPropertyHandle>
 			.OnValueChanged(this, &FDMXAttributeNameCustomization::SetValue)
 		];
 
-	ProtocolSettings->GetOnDefaultAttributesChanged().AddSP(this, &FDMXAttributeNameCustomization::ForceRefresh);
+	ProtocolSettings->GetOnDefaultAttributesChanged().AddSP(this, &FDMXAttributeNameCustomization::RequestRefresh);
+
+	NameHandle->SetOnPropertyValueChanged(FSimpleDelegate::CreateSP(this, &FDMXAttributeNameCustomization::RequestRefresh));
 }
 
 void FDMXAttributeNameCustomization::CustomizeChildren(TSharedRef<IPropertyHandle> InPropertyHandle, IDetailChildrenBuilder& InChildBuilder, IPropertyTypeCustomizationUtils& CustomizationUtils)
@@ -70,19 +73,10 @@ void FDMXAttributeNameCustomization::CustomizeChildren(TSharedRef<IPropertyHandl
 
 FName FDMXAttributeNameCustomization::GetValue() const
 {
-	TArray<const void*> RawData;
-	StructPropertyHandle->AccessRawData(RawData);
-
-	for (const void* RawPtr : RawData)
-	{
-		if (RawPtr != nullptr)
-		{
-			// The types we use with this customization must have a cast constructor to FName
-			return reinterpret_cast<const FDMXAttributeName*>(RawPtr)->Name;
-		}
-	}
-
-	return FName();
+	FName Name;
+	NameHandle->GetValue(Name);
+	
+	return Name;
 }
 
 void FDMXAttributeNameCustomization::SetValue(FName NewValue)
@@ -129,7 +123,16 @@ bool FDMXAttributeNameCustomization::HasMultipleValues() const
 	return false;
 }
 
-void FDMXAttributeNameCustomization::ForceRefresh()
+void FDMXAttributeNameCustomization::RequestRefresh()
+{
+	if (!RefreshHandle.IsValid())
+	{
+		RefreshHandle = GEditor->GetTimerManager()->SetTimerForNextTick(FTimerDelegate::CreateSP(this, &FDMXAttributeNameCustomization::Refresh));
+	}
+}
+
+void FDMXAttributeNameCustomization::Refresh()
 {
+	RefreshHandle.Invalidate();
 	PropertyUtilities->ForceRefresh();
 }
diff --git a/Engine/Plugins/VirtualProduction/DMX/DMXEngine/Source/DMXEditor/Private/Customizations/DMXAttributeNameCustomization.h b/Engine/Plugins/VirtualProduction/DMX/DMXEngine/Source/DMXEditor/Private/Customizations/DMXAttributeNameCustomization.h
--- a/Engine/Plugins/VirtualProduction/DMX/DMXEngine/Source/DMXEditor/Private/Customizations/DMXAttributeNameCustomization.h
+++ b/Engine/Plugins/VirtualProduction/DMX/DMXEngine/Source/DMXEditor/Private/Customizations/DMXAttributeNameCustomization.h
@@ -6,6 +6,7 @@
 
 #include "CoreMinimal.h"
 #include "DetailWidgetRow.h"
+#include "Engine/EngineTypes.h"
 #include "IPropertyTypeCustomization.h"
 #include "IPropertyUtilities.h"
 #include "Widgets/SNameListPicker.h"
@@ -38,8 +39,11 @@ private:
 	/** True if the customization displays multiple values */
 	bool HasMultipleValues() const;
 
-	/** Forces a refresh of the customization */
-	void ForceRefresh();
+	/** Requests a refresh of the customization */
+	void RequestRefresh();
+
+	/** Refreshes instantly */
+	void Refresh();
 
 	/** Property handle for the Name property */
 	TSharedPtr<IPropertyHandle> NameHandle;
@@ -47,6 +51,9 @@ private:
 	/** Handle for the FDMXAttributeName struct self */
 	TSharedPtr<IPropertyHandle> StructPropertyHandle;
 
+	/** Timer handle set when refresh is requested but not carried out yet */
+	FTimerHandle RefreshHandle;
+
 	/** Property utlities for this customization */
 	TSharedPtr<IPropertyUtilities> PropertyUtilities;
 };
diff --git a/Engine/Plugins/VirtualProduction/DMX/DMXPixelMapping/Source/DMXPixelMappingEditor/Private/Views/SDMXPixelMappingHierarchyView.cpp b/Engine/Plugins/VirtualProduction/DMX/DMXPixelMapping/Source/DMXPixelMappingEditor/Private/Views/SDMXPixelMappingHierarchyView.cpp
--- a/Engine/Plugins/VirtualProduction/DMX/DMXPixelMapping/Source/DMXPixelMappingEditor/Private/Views/SDMXPixelMappingHierarchyView.cpp
+++ b/Engine/Plugins/VirtualProduction/DMX/DMXPixelMapping/Source/DMXPixelMappingEditor/Private/Views/SDMXPixelMappingHierarchyView.cpp
@@ -118,10 +118,7 @@ void SDMXPixelMappingHierarchyView::Construct(const FArguments& InArgs, const TS
 			]
 		];
 
-	RebuildTreeView();
-
-	bRefreshRequested = true;
-
+	bRebuildTreeRequested = true;
 	bSelectFirstRenderer = true;
 
 	GEditor->RegisterForUndo(this);
@@ -136,7 +133,15 @@ void SDMXPixelMappingHierarchyView::Construct(const FArguments& InArgs, const TS
 
 void SDMXPixelMappingHierarchyView::SelectFirstAvailableRenderer()
 {
-	if (TreeRootWidgets.Num() > 0)
+	bSelectFirstRenderer = true;
+}
+
+void SDMXPixelMappingHierarchyView::Tick(const FGeometry& AllottedGeometry, const double InCurrentTime, const float InDeltaTime)
+{
+	ConditionallyUpdateTree();
+
+	// select the first available renderer, if any is available
+	if (bSelectFirstRenderer && TreeRootWidgets.Num() > 0 && WidgetTreeView->GetNumItemsSelected() == 0)
 	{
 		FDMXPixelMappingHierarchyItemWidgetModelArr RootComponents;
 		TreeRootWidgets[0]->GatherChildren(RootComponents);
@@ -144,29 +149,15 @@ void SDMXPixelMappingHierarchyView::SelectFirstAvailableRenderer()
 		{
 			if (UDMXPixelMappingRendererComponent* RendererComponent = Cast<UDMXPixelMappingRendererComponent>(ItemWidgetPtr->GetReference().GetComponent()))
 			{
-				CurrentSelectedItems.Empty();
-				CurrentSelectedItems.Add(ItemWidgetPtr);
-				bRefreshRequested = true;
+				WidgetTreeView->SetSelection(ItemWidgetPtr, ESelectInfo::OnMouseClick);
+				bSelectFirstRenderer = false;
+
 				break;
 			}
 		}
 	}
 }
 
-void SDMXPixelMappingHierarchyView::Tick(const FGeometry& AllottedGeometry, const double InCurrentTime, const float InDeltaTime)
-{
-
-	// select the first available renderer (if any)
-	if (bSelectFirstRenderer && TreeRootWidgets.Num() > 0 && CurrentSelectedItems.Num() == 0)
-	{
-		SelectFirstAvailableRenderer();
-		bSelectFirstRenderer = false;
-	}
-
-	ConditionallyUpdateTree();
-}
-
-
 FReply SDMXPixelMappingHierarchyView::OnKeyDown(const FGeometry& MyGeometry, const FKeyEvent& InKeyEvent)
 {
 	if (CommandList->ProcessCommandBindings(InKeyEvent))
@@ -314,12 +305,6 @@ void SDMXPixelMappingHierarchyView::ConditionallyUpdateTree()
 		RefreshTree();
 		ExpansionSnapshot.RestoreExpandedAndExpandNewModels(RootWidgets, WidgetTreeView);
 
-		for (FDMXPixelMappingHierarchyItemWidgetModelPtr& Item : CurrentSelectedItems)
-		{
-			Item->OnSelection();
-		}
-		CurrentSelectedItems.Empty();
-
 		bRefreshRequested = false;
 		bRebuildTreeRequested = false;
 
@@ -359,33 +344,59 @@ void SDMXPixelMappingHierarchyView::RebuildTreeView()
 
 void SDMXPixelMappingHierarchyView::RestoreSelectedItems()
 {
-	for (FDMXPixelMappingHierarchyItemWidgetModelPtr& Model : RootWidgets)
+	if (bIsUpdatingSelection)
 	{
-		RecursiveSelection(Model);
+		return;
 	}
+
+	GEditor->GetTimerManager()->SetTimerForNextTick([this]()
+		{
+
+			WidgetTreeView->ClearSelection();
+			for (FDMXPixelMappingHierarchyItemWidgetModelPtr& Model : RootWidgets)
+			{
+				RestoreSelectionForItemAndChildren(Model);
+			}
+		});
 }
 
-void SDMXPixelMappingHierarchyView::RecursiveSelection(FDMXPixelMappingHierarchyItemWidgetModelPtr& Model)
+bool SDMXPixelMappingHierarchyView::RestoreSelectionForItemAndChildren(FDMXPixelMappingHierarchyItemWidgetModelPtr& Model)
 {
-	if (Model->ContainsSelection())
+	if (bIsUpdatingSelection)
 	{
-		// Expand items that contain selection.
-		WidgetTreeView->SetItemExpansion(Model, true);
+		return false;
+	}
 
-		TArray<FDMXPixelMappingHierarchyItemWidgetModelPtr> Children;
-		Model->GatherChildren(Children);
+	const TSharedPtr<FDMXPixelMappingToolkit> StrongToolkit = Toolkit.Pin();
+	if (!StrongToolkit.IsValid())
+	{
+		return false;
+	}
 
-		for (FDMXPixelMappingHierarchyItemWidgetModelPtr& ChildModel : Children)
-		{
-			RecursiveSelection(ChildModel);
-		}
+	TArray<FDMXPixelMappingHierarchyItemWidgetModelPtr> Children;
+	Model->GatherChildren(Children);
+
+	bool bContainsSelection = false;
+	for (FDMXPixelMappingHierarchyItemWidgetModelPtr& ChildModel : Children)
+	{
+		bContainsSelection |= RestoreSelectionForItemAndChildren(ChildModel);
+	}
+
+	const TSet<FDMXPixelMappingComponentReference> SelectedComponents = StrongToolkit->GetSelectedComponents();
+	if (bContainsSelection)
+	{
+		WidgetTreeView->SetItemExpansion(Model, true);
 	}
 
-	if (Model->IsSelected())
+	if (SelectedComponents.Contains(Model->GetReference()))
 	{
 		WidgetTreeView->SetItemSelection(Model, true, ESelectInfo::Direct);
 		WidgetTreeView->RequestScrollIntoView(Model);
+		
+		return true;
 	}
+
+	return bContainsSelection;
 }
 
 void SDMXPixelMappingHierarchyView::RefreshTree()
@@ -411,13 +422,6 @@ void SDMXPixelMappingHierarchyView::OnEditorSelectionChanged()
 {
 	if (!bIsUpdatingSelection)
 	{
-		WidgetTreeView->ClearSelection();
-
-		if (RootWidgets.Num() > 0)
-		{
-			RootWidgets[0]->RefreshSelection();
-		}
-
 		RestoreSelectedItems();
 	}
 }
@@ -583,11 +587,6 @@ void SDMXPixelMappingHierarchyView::BeginPaste()
 			MoveComponentToComponent(Component, SelectedComponent, true);
 		}
 	}
-
-	if (bRebuildTreeRequested)
-	{
-		CurrentSelectedItems = WidgetTreeView->GetSelectedItems();
-	}
 }
 
 bool SDMXPixelMappingHierarchyView::CanBeginDuplicate() const
diff --git a/Engine/Plugins/VirtualProduction/DMX/DMXPixelMapping/Source/DMXPixelMappingEditor/Private/Views/SDMXPixelMappingHierarchyView.h b/Engine/Plugins/VirtualProduction/DMX/DMXPixelMapping/Source/DMXPixelMappingEditor/Private/Views/SDMXPixelMappingHierarchyView.h
--- a/Engine/Plugins/VirtualProduction/DMX/DMXPixelMapping/Source/DMXPixelMappingEditor/Private/Views/SDMXPixelMappingHierarchyView.h
+++ b/Engine/Plugins/VirtualProduction/DMX/DMXPixelMapping/Source/DMXPixelMappingEditor/Private/Views/SDMXPixelMappingHierarchyView.h
@@ -82,7 +82,8 @@ private:
 
 	void RestoreSelectedItems();
 
-	void RecursiveSelection(FDMXPixelMappingHierarchyItemWidgetModelPtr& Model);
+	/** Restores selection for a the item and its children. Returns true if the Model or a child was selected. */
+	bool RestoreSelectionForItemAndChildren(FDMXPixelMappingHierarchyItemWidgetModelPtr& Model);
 
 	/** Rebuilds the tree structure based on the current filter options */
 	void RefreshTree();
@@ -156,7 +157,5 @@ private:
 
 	TreeViewPtr WidgetTreeView;
 
-	FDMXPixelMappingHierarchyItemWidgetModelArr CurrentSelectedItems;
-
 	bool bSelectFirstRenderer;
 };
diff --git a/Engine/Plugins/VirtualProduction/DMX/DMXPixelMapping/Source/DMXPixelMappingEditor/Private/Views/SDMXPixelMappingPreviewView.cpp b/Engine/Plugins/VirtualProduction/DMX/DMXPixelMapping/Source/DMXPixelMappingEditor/Private/Views/SDMXPixelMappingPreviewView.cpp
--- a/Engine/Plugins/VirtualProduction/DMX/DMXPixelMapping/Source/DMXPixelMappingEditor/Private/Views/SDMXPixelMappingPreviewView.cpp
+++ b/Engine/Plugins/VirtualProduction/DMX/DMXPixelMapping/Source/DMXPixelMappingEditor/Private/Views/SDMXPixelMappingPreviewView.cpp
@@ -106,8 +106,6 @@ void SDMXPixelMappingPreviewView::Construct(const FArguments& InArgs, const TSha
 			]
 		]
 	, InToolkit);
-
-	ZoomToFit(true);
 }
 
 FReply SDMXPixelMappingPreviewView::OnMouseButtonDown(const FGeometry& MyGeometry, const FPointerEvent& MouseEvent)
diff --git a/Engine/Plugins/VirtualProduction/DMX/DMXPixelMapping/Source/DMXPixelMappingEditor/Private/Widgets/SDMXPixelMappingSurface.cpp b/Engine/Plugins/VirtualProduction/DMX/DMXPixelMapping/Source/DMXPixelMappingEditor/Private/Widgets/SDMXPixelMappingSurface.cpp
--- a/Engine/Plugins/VirtualProduction/DMX/DMXPixelMapping/Source/DMXPixelMappingEditor/Private/Widgets/SDMXPixelMappingSurface.cpp
+++ b/Engine/Plugins/VirtualProduction/DMX/DMXPixelMapping/Source/DMXPixelMappingEditor/Private/Widgets/SDMXPixelMappingSurface.cpp
@@ -172,6 +172,8 @@ void SDMXPixelMappingSurface::Construct(const FArguments& InArgs, const TSharedP
 		[
 			InArgs._Content.Widget
 		];
+
+	InitialBounds = ComputeAreaBounds();
 }
 
 EActiveTimerReturnType SDMXPixelMappingSurface::HandleZoomToFit(double InCurrentTime, float InDeltaTime)
@@ -209,7 +211,7 @@ void SDMXPixelMappingSurface::Tick(const FGeometry& AllottedGeometry, const doub
 		}
 	}
 
-	if (!bRecenteredOnFirstTick)
+	if (!bRecenteredOnFirstTick && InitialBounds != ComputeAreaBounds())
 	{
 		ZoomToFit(true);
 		bRecenteredOnFirstTick = true;
diff --git a/Engine/Plugins/VirtualProduction/DMX/DMXPixelMapping/Source/DMXPixelMappingEditor/Private/Widgets/SDMXPixelMappingSurface.h b/Engine/Plugins/VirtualProduction/DMX/DMXPixelMapping/Source/DMXPixelMappingEditor/Private/Widgets/SDMXPixelMappingSurface.h
--- a/Engine/Plugins/VirtualProduction/DMX/DMXPixelMapping/Source/DMXPixelMappingEditor/Private/Widgets/SDMXPixelMappingSurface.h
+++ b/Engine/Plugins/VirtualProduction/DMX/DMXPixelMapping/Source/DMXPixelMappingEditor/Private/Widgets/SDMXPixelMappingSurface.h
@@ -151,6 +151,9 @@ protected:
 	/** Does the user need to press Control in order to over-zoom. */
 	bool bRequireControlToOverZoom;
 
+	/** Initial bounds, useful to detect when the initial zoom to fit should occur */
+	FSlateRect InitialBounds;
+
 	/** Toolkit weak pointer */
 	TWeakPtr<FDMXPixelMappingToolkit> ToolkitWeakPtr;
 
diff --git a/Engine/Plugins/VirtualProduction/DMX/DMXPixelMapping/Source/DMXPixelMappingRuntime/Private/Components/DMXPixelMappingMatrixComponent.cpp b/Engine/Plugins/VirtualProduction/DMX/DMXPixelMapping/Source/DMXPixelMappingRuntime/Private/Components/DMXPixelMappingMatrixComponent.cpp
--- a/Engine/Plugins/VirtualProduction/DMX/DMXPixelMapping/Source/DMXPixelMappingRuntime/Private/Components/DMXPixelMappingMatrixComponent.cpp
+++ b/Engine/Plugins/VirtualProduction/DMX/DMXPixelMapping/Source/DMXPixelMappingRuntime/Private/Components/DMXPixelMappingMatrixComponent.cpp
@@ -126,6 +126,9 @@ void UDMXPixelMappingMatrixComponent::LogInvalidProperties()
 #if WITH_EDITOR
 void UDMXPixelMappingMatrixComponent::PreEditChange(FProperty* PropertyAboutToChange)
 {
+	// Explicitly recording the state here is required since many properties propagonate to the child cells and cause otherwise malformed transactions
+	Modify();
+
 	if (PropertyAboutToChange)
 	{
 		const FName PropertyName = PropertyAboutToChange->GetFName();
