DMX - Fix Pixel Mapper layout menu item has no icon

diff --git a/Engine/Plugins/VirtualProduction/DMX/DMXPixelMapping/Source/DMXPixelMappingEditor/Private/Toolkits/DMXPixelMappingToolkit.cpp b/Engine/Plugins/VirtualProduction/DMX/DMXPixelMapping/Source/DMXPixelMappingEditor/Private/Toolkits/DMXPixelMappingToolkit.cpp
--- a/Engine/Plugins/VirtualProduction/DMX/DMXPixelMapping/Source/DMXPixelMappingEditor/Private/Toolkits/DMXPixelMappingToolkit.cpp
+++ b/Engine/Plugins/VirtualProduction/DMX/DMXPixelMapping/Source/DMXPixelMappingEditor/Private/Toolkits/DMXPixelMappingToolkit.cpp
@@ -104,7 +104,8 @@ void FDMXPixelMappingToolkit::RegisterTabSpawners(const TSharedRef<class FTabMan
 
 	InTabManager->RegisterTabSpawner(LayoutViewTabID, FOnSpawnTab::CreateSP(this, &FDMXPixelMappingToolkit::SpawnTab_LayoutView))
 		.SetDisplayName(LOCTEXT("Tab_LayoutView", "Layout"))
-		.SetGroup(WorkspaceMenuCategoryRef);
+		.SetGroup(WorkspaceMenuCategoryRef)
+		.SetIcon(FSlateIcon(FAppStyle::GetAppStyleSetName(), "Icons.Layout"));
 }
 
 void FDMXPixelMappingToolkit::UnregisterTabSpawners(const TSharedRef<class FTabManager>& InTabManager)
