DMX - Add Color Space management to PixelMapping

diff --git a/Engine/Plugins/VirtualProduction/DMX/DMXPixelMapping/Shaders/Private/DMXPixelMapping.usf b/Engine/Plugins/VirtualProduction/DMX/DMXPixelMapping/Shaders/Private/DMXPixelMapping.usf
--- a/Engine/Plugins/VirtualProduction/DMX/DMXPixelMapping/Shaders/Private/DMXPixelMapping.usf
+++ b/Engine/Plugins/VirtualProduction/DMX/DMXPixelMapping/Shaders/Private/DMXPixelMapping.usf
@@ -85,18 +85,6 @@ OutputPS DMXPixelMappingPS(OutputVS IN)
 		OUT.Color /= (NUM_CELLS + 1);
 		
 #endif
-
-		// Invert color
-		for (int ColorIndex = 0; ColorIndex < 4; ColorIndex++)
-		{
-			if (InvertPixel[ColorIndex])
-			{
-				OUT.Color[ColorIndex] = 1.f - OUT.Color[ColorIndex];
-			}
-		}
-		
-		// Apply color factor
-		OUT.Color *= PixelFactor;
 	}
 			
 	return OUT;
diff --git a/Engine/Plugins/VirtualProduction/DMX/DMXPixelMapping/Shaders/Private/DMXPixelMapping.ush b/Engine/Plugins/VirtualProduction/DMX/DMXPixelMapping/Shaders/Private/DMXPixelMapping.ush
--- a/Engine/Plugins/VirtualProduction/DMX/DMXPixelMapping/Shaders/Private/DMXPixelMapping.ush
+++ b/Engine/Plugins/VirtualProduction/DMX/DMXPixelMapping/Shaders/Private/DMXPixelMapping.ush
@@ -19,8 +19,6 @@ SamplerState InputSampler;
 // PS Parameters
 int2 InputTextureSize;
 int2 OutputTextureSize;
-float4 PixelFactor;
-int4 InvertPixel;
 float2 UVCellSize;
 
 struct InputVS
diff --git a/Engine/Plugins/VirtualProduction/DMX/DMXPixelMapping/Source/DMXPixelMappingEditor/DMXPixelMappingEditor.Build.cs b/Engine/Plugins/VirtualProduction/DMX/DMXPixelMapping/Source/DMXPixelMappingEditor/DMXPixelMappingEditor.Build.cs
--- a/Engine/Plugins/VirtualProduction/DMX/DMXPixelMapping/Source/DMXPixelMappingEditor/DMXPixelMappingEditor.Build.cs
+++ b/Engine/Plugins/VirtualProduction/DMX/DMXPixelMapping/Source/DMXPixelMappingEditor/DMXPixelMappingEditor.Build.cs
@@ -7,6 +7,14 @@ public class DMXPixelMappingEditor : ModuleRules
 {
 	public DMXPixelMappingEditor(ReadOnlyTargetRules Target) : base(Target)
 	{
+		var EngineDir = Path.GetFullPath(Target.RelativeEnginePath);
+
+		PrivateIncludePaths.AddRange(
+			new string[] {
+				System.IO.Path.Combine(GetModuleDirectory("DMXPixelMappingRuntime"), "Private"),
+			}
+		);
+
 		PrivateDependencyModuleNames.AddRange( new string[] {
 			"ApplicationCore",
 			"Core",
diff --git a/Engine/Plugins/VirtualProduction/DMX/DMXPixelMapping/Source/DMXPixelMappingEditor/Private/Customizations/DMXPixelMappingColorSpaceDetails_RGBCMY.cpp b/Engine/Plugins/VirtualProduction/DMX/DMXPixelMapping/Source/DMXPixelMappingEditor/Private/Customizations/DMXPixelMappingColorSpaceDetails_RGBCMY.cpp
new file mode 100644
--- /dev/null
+++ b/Engine/Plugins/VirtualProduction/DMX/DMXPixelMapping/Source/DMXPixelMappingEditor/Private/Customizations/DMXPixelMappingColorSpaceDetails_RGBCMY.cpp
@@ -0,0 +1,141 @@
+// Copyright Epic Games, Inc. All Rights Reserved.
+
+#include "Customizations/DMXPixelMappingColorSpaceDetails_RGBCMY.h"
+
+#include "ColorSpace/DMXPixelMappingColorSpace_RGBCMY.h"
+
+#include "DetailCategoryBuilder.h"
+#include "DetailLayoutBuilder.h"
+#include "DetailWidgetRow.h"
+#include "PropertyHandle.h"
+#include "Widgets/Layout/SBox.h"
+#include "Widgets/Layout/SSeparator.h"
+#include "Widgets/Layout/SWrapBox.h"
+#include "Widgets/Text/STextBlock.h"
+
+
+#define LOCTEXT_NAMESPACE "DMXPixelMappingColorSpaceDetails_RGBCMY"
+
+TSharedRef<IDetailCustomization> FDMXPixelMappingColorSpaceDetails_RGBCMY::MakeInstance()
+{
+	return MakeShared<FDMXPixelMappingColorSpaceDetails_RGBCMY>();
+}
+
+void FDMXPixelMappingColorSpaceDetails_RGBCMY::CustomizeDetails(IDetailLayoutBuilder& DetailBuilder)
+{
+	// Hide bCyan, bMagenta, bYellow properties, instead show them inline
+	SendCyanHandle = DetailBuilder.GetProperty(GET_MEMBER_NAME_CHECKED(UDMXPixelMappingColorSpace_RGBCMY, bSendCyan));
+	SendMagentaHandle = DetailBuilder.GetProperty(GET_MEMBER_NAME_CHECKED(UDMXPixelMappingColorSpace_RGBCMY, bSendMagenta));
+	SendYellowHandle = DetailBuilder.GetProperty(GET_MEMBER_NAME_CHECKED(UDMXPixelMappingColorSpace_RGBCMY, bSendYellow));
+	SendCyanHandle->MarkHiddenByCustomization();
+	SendMagentaHandle->MarkHiddenByCustomization();
+	SendYellowHandle->MarkHiddenByCustomization();
+
+	FDMXAttributeRowData RedAttributeRowData;
+	RedAttributeRowData.AttributeHandle = DetailBuilder.GetProperty(GET_MEMBER_NAME_CHECKED(UDMXPixelMappingColorSpace_RGBCMY, RedAttribute));
+	RedAttributeRowData.AttributeLabel = TAttribute<FText>::CreateLambda([this]() { return GetRedAttributeDisplayName(); });
+	RedAttributeRowData.IsInvertColorHandle = SendCyanHandle;
+	RedAttributeRowData.InvertColorLabel = LOCTEXT("InvertRedToCyanLabel", "Send Cyan");
+	GenerateAttributeRow(DetailBuilder, RedAttributeRowData);
+
+	FDMXAttributeRowData GreenAttributeRowData;
+	GreenAttributeRowData.AttributeHandle = DetailBuilder.GetProperty(GET_MEMBER_NAME_CHECKED(UDMXPixelMappingColorSpace_RGBCMY, GreenAttribute));
+	GreenAttributeRowData.AttributeLabel = TAttribute<FText>::CreateLambda([this]() { return GetGreenAttributeDisplayName(); });
+	GreenAttributeRowData.IsInvertColorHandle = SendMagentaHandle;
+	GreenAttributeRowData.InvertColorLabel = LOCTEXT("InvertGreenToMagentaLabel", "Send Magenta");
+	GenerateAttributeRow(DetailBuilder, GreenAttributeRowData);
+
+	FDMXAttributeRowData BlueAttributeRowData;
+	BlueAttributeRowData.AttributeHandle = DetailBuilder.GetProperty(GET_MEMBER_NAME_CHECKED(UDMXPixelMappingColorSpace_RGBCMY, BlueAttribute));
+	BlueAttributeRowData.AttributeLabel = TAttribute<FText>::CreateLambda([this]() { return GetBlueAttributeDisplayName(); });
+	BlueAttributeRowData.IsInvertColorHandle = SendYellowHandle;
+	BlueAttributeRowData.InvertColorLabel = LOCTEXT("InvertBlueToMagentaLabel", "Send Yellow");
+	BlueAttributeRowData.bAppendSeparator = true;
+	GenerateAttributeRow(DetailBuilder, BlueAttributeRowData);
+}
+
+void FDMXPixelMappingColorSpaceDetails_RGBCMY::GenerateAttributeRow(IDetailLayoutBuilder& DetailBuilder, const FDMXAttributeRowData& AttributeRowData)
+{
+	IDetailCategoryBuilder& ColorSpaceCategory = DetailBuilder.EditCategory("RGB");
+	IDetailPropertyRow& Row = ColorSpaceCategory.AddProperty(AttributeRowData.AttributeHandle);
+
+	TSharedPtr<SWidget> NameWidget;
+	TSharedPtr<SWidget> ValueWidget;
+	FDetailWidgetRow DefaultWidgetRow;
+	Row.GetDefaultWidgets(NameWidget, ValueWidget, DefaultWidgetRow);
+
+	Row.CustomWidget()
+		.NameContent()
+		[
+			SNew(STextBlock)
+			.Font(DetailBuilder.GetDetailFont())
+			.Text(AttributeRowData.AttributeLabel)
+		]
+		.ValueContent()
+		.HAlign(HAlign_Fill)
+		[
+			SNew(SWrapBox)
+			.UseAllottedWidth(true)
+			.InnerSlotPadding(FVector2D(4.f, 0.f))
+
+			+ SWrapBox::Slot()
+			.HAlign(HAlign_Fill)
+			[
+				SNew(SBox)
+				.MinDesiredWidth(120.f)
+				[
+					ValueWidget.ToSharedRef()
+				]
+			]
+
+			+ SWrapBox::Slot()
+			.HAlign(HAlign_Left)
+			.Padding(4.f, 0.f, 0.f, 0.f)
+			[
+				AttributeRowData.IsInvertColorHandle->CreatePropertyValueWidget()
+			]
+
+			+ SWrapBox::Slot()
+			.HAlign(HAlign_Left)
+			.VAlign(VAlign_Center)
+			[
+				SNew(STextBlock)
+				.Font(DetailBuilder.GetDetailFont())
+				.Text(AttributeRowData.InvertColorLabel)
+			]
+		];
+
+	
+}
+
+FText FDMXPixelMappingColorSpaceDetails_RGBCMY::GetRedAttributeDisplayName() const
+{
+	bool bIsCyan = false;
+	SendCyanHandle->GetValue(bIsCyan);
+
+	return bIsCyan ?
+		LOCTEXT("RedAsCyanAttributeDisplayName", "Cyan Attribute") :
+		LOCTEXT("RedAttributeDisplayName", "Red Attribute");
+}
+
+FText FDMXPixelMappingColorSpaceDetails_RGBCMY::GetGreenAttributeDisplayName() const
+{
+	bool bIsMagenta = false;
+	SendMagentaHandle->GetValue(bIsMagenta);
+
+	return bIsMagenta ?
+		LOCTEXT("GreenAsMagentaAttributeDisplayName", "Magenta Attribute") :
+		LOCTEXT("GreenAttributeDisplayName", "Green Attribute");
+}
+
+FText FDMXPixelMappingColorSpaceDetails_RGBCMY::GetBlueAttributeDisplayName() const
+{
+	bool bIsYellow = false;
+	SendYellowHandle->GetValue(bIsYellow);
+
+	return bIsYellow ?
+		LOCTEXT("BlueAsYellowAttributeDisplayName", "Yellow Attribute") :
+		LOCTEXT("BlueAttributeDisplayName", "Blue Attribute");
+}
+
+#undef LOCTEXT_NAMESPACE
diff --git a/Engine/Plugins/VirtualProduction/DMX/DMXPixelMapping/Source/DMXPixelMappingEditor/Private/Customizations/DMXPixelMappingColorSpaceDetails_RGBCMY.h b/Engine/Plugins/VirtualProduction/DMX/DMXPixelMapping/Source/DMXPixelMappingEditor/Private/Customizations/DMXPixelMappingColorSpaceDetails_RGBCMY.h
new file mode 100644
--- /dev/null
+++ b/Engine/Plugins/VirtualProduction/DMX/DMXPixelMapping/Source/DMXPixelMappingEditor/Private/Customizations/DMXPixelMappingColorSpaceDetails_RGBCMY.h
@@ -0,0 +1,58 @@
+// Copyright Epic Games, Inc. All Rights Reserved.
+
+#pragma once
+
+#include "IDetailCustomization.h"
+#include "Internationalization/Text.h"
+#include "Misc/Attribute.h"
+#include "Templates/SharedPointer.h"
+
+class IDetailLayoutBuilder;
+class IPropertyHandle;
+
+
+/** Details customization for the RGBCMY Color Space */
+class FDMXPixelMappingColorSpaceDetails_RGBCMY
+	: public IDetailCustomization
+{
+public:
+	/** Creates an instance of this customization */
+	static TSharedRef<IDetailCustomization> MakeInstance();
+
+	//~ Begin IDetailCustomization interface 
+	virtual void CustomizeDetails(IDetailLayoutBuilder& DetailBuilder) override;
+	//~ End IDetailCustomization interface 
+
+private:
+	/** Holds required data to create an Attribute row */
+	struct FDMXAttributeRowData
+	{
+		TSharedPtr<IPropertyHandle> AttributeHandle;
+		TAttribute<FText> AttributeLabel;
+		TSharedPtr<IPropertyHandle> IsInvertColorHandle;
+		FText InvertColorLabel;
+		bool bAppendSeparator = false;
+	};
+
+	/** Generates the row for one of the Attributes */
+	void GenerateAttributeRow(IDetailLayoutBuilder& DetailBuilder, const FDMXAttributeRowData& AttributeRowData);
+
+	/** Returns the Display Name of the Red Attribute */
+	FText GetRedAttributeDisplayName() const;
+
+	/** Returns the Display Name of the Red Attribute */
+	FText GetGreenAttributeDisplayName() const;
+
+	/** Returns the Display Name of the Red Attribute */
+	FText GetBlueAttributeDisplayName() const;
+
+	/** Property Handle for the bSendCyan Property */
+	TSharedPtr<IPropertyHandle> SendCyanHandle;
+
+	/** Property Handle for the bSendMagenta Property */
+	TSharedPtr<IPropertyHandle> SendMagentaHandle;
+
+	/** Property Handle for the bSendYellow Property */
+	TSharedPtr<IPropertyHandle> SendYellowHandle;
+};
+
diff --git a/Engine/Plugins/VirtualProduction/DMX/DMXPixelMapping/Source/DMXPixelMappingEditor/Private/Customizations/DMXPixelMappingDetailCustomization_FixtureGroup.cpp b/Engine/Plugins/VirtualProduction/DMX/DMXPixelMapping/Source/DMXPixelMappingEditor/Private/Customizations/DMXPixelMappingDetailCustomization_FixtureGroup.cpp
--- a/Engine/Plugins/VirtualProduction/DMX/DMXPixelMapping/Source/DMXPixelMappingEditor/Private/Customizations/DMXPixelMappingDetailCustomization_FixtureGroup.cpp
+++ b/Engine/Plugins/VirtualProduction/DMX/DMXPixelMapping/Source/DMXPixelMappingEditor/Private/Customizations/DMXPixelMappingDetailCustomization_FixtureGroup.cpp
@@ -2,7 +2,6 @@
 
 #include "Customizations/DMXPixelMappingDetailCustomization_FixtureGroup.h"
 
-#include "Algo/Find.h"
 #include "DMXPixelMapping.h"
 #include "DMXPixelMappingLayoutSettings.h"
 #include "Components/DMXPixelMappingFixtureGroupComponent.h"
@@ -19,6 +18,7 @@
 #include "DetailWidgetRow.h"
 #include "IPropertyUtilities.h"
 #include "ScopedTransaction.h"
+#include "Algo/Find.h"
 #include "Misc/CoreDelegates.h"
 #include "Widgets/Input/SButton.h"
 
@@ -84,6 +84,8 @@ void FDMXPixelMappingDetailCustomization_FixtureGroup::CreateAddAllPatchesButton
 
 void FDMXPixelMappingDetailCustomization_FixtureGroup::CreateFixturePatchDetailRows(IDetailLayoutBuilder& InDetailLayout)
 {
+	
+
 	UDMXPixelMappingFixtureGroupComponent* FixtureGroupComponent = WeakFixtureGroupComponent.Get();
 	if (!FixtureGroupComponent)
 	{
@@ -103,6 +105,7 @@ void FDMXPixelMappingDetailCustomization_FixtureGroup::CreateFixturePatchDetailR
 	EntitiesHandle->SetOnPropertyValueChanged(FSimpleDelegate::CreateSP(this, &FDMXPixelMappingDetailCustomization_FixtureGroup::ForceRefresh));
 
 	// Add fixture patches as custom rows
+	NumFixturePatchRows = 0;
 	TArray<UDMXEntityFixturePatch*> AllFixturePatches = DMXLibrary->GetEntitiesTypeCast<UDMXEntityFixturePatch>();
 	for (UDMXEntityFixturePatch* FixturePatch : AllFixturePatches)
 	{
diff --git a/Engine/Plugins/VirtualProduction/DMX/DMXPixelMapping/Source/DMXPixelMappingEditor/Private/Customizations/DMXPixelMappingDetailCustomization_FixtureGroupItem.cpp b/Engine/Plugins/VirtualProduction/DMX/DMXPixelMapping/Source/DMXPixelMappingEditor/Private/Customizations/DMXPixelMappingDetailCustomization_FixtureGroupItem.cpp
--- a/Engine/Plugins/VirtualProduction/DMX/DMXPixelMapping/Source/DMXPixelMappingEditor/Private/Customizations/DMXPixelMappingDetailCustomization_FixtureGroupItem.cpp
+++ b/Engine/Plugins/VirtualProduction/DMX/DMXPixelMapping/Source/DMXPixelMappingEditor/Private/Customizations/DMXPixelMappingDetailCustomization_FixtureGroupItem.cpp
@@ -9,7 +9,6 @@
 
 #include "DetailLayoutBuilder.h"
 #include "DetailWidgetRow.h"
-#include "Framework/Views/TableViewMetadata.h"
 #include "IPropertyUtilities.h"
 #include "IDetailsView.h"
 #include "Modules/ModuleManager.h"
@@ -17,24 +16,23 @@
 
 #define LOCTEXT_NAMESPACE "FixtureGroupItem"
 
+
 void FDMXPixelMappingDetailCustomization_FixtureGroupItem::CustomizeDetails(IDetailLayoutBuilder& InDetailLayout)
 {
-	DetailLayout = &InDetailLayout;
 	PropertyUtilities = InDetailLayout.GetPropertyUtilities();
 
 	// Get editing object
 	TArray<TWeakObjectPtr<UObject>> Objects;
-	DetailLayout->GetObjectsBeingCustomized(Objects);
+	InDetailLayout.GetObjectsBeingCustomized(Objects);
 
 	FixtureGroupItemComponents.Empty();
-	
 	for (TWeakObjectPtr<UObject> SelectedObject : Objects)
 	{
 		FixtureGroupItemComponents.Add(Cast<UDMXPixelMappingFixtureGroupItemComponent>(SelectedObject));
 	}
 
 	// Get editing categories
-	IDetailCategoryBuilder& OutputSettingsCategory = DetailLayout->EditCategory("Output Settings", FText::GetEmpty(), ECategoryPriority::Important);
+	IDetailCategoryBuilder& OutputSettingsCategory = InDetailLayout.EditCategory("Output Settings", FText::GetEmpty(), ECategoryPriority::Important);
 
 	// Hide absolute postition property handles
 	TSharedPtr<IPropertyHandle> PositionXPropertyHandle = InDetailLayout.GetProperty(UDMXPixelMappingOutputComponent::GetPositionXPropertyName());
@@ -46,214 +44,62 @@ void FDMXPixelMappingDetailCustomization_FixtureGroupItem::CustomizeDetails(IDet
 	TSharedPtr<IPropertyHandle> SizeYPropertyHandle = InDetailLayout.GetProperty(UDMXPixelMappingOutputComponent::GetSizeYPropertyName());
 	InDetailLayout.HideProperty(SizeXPropertyHandle);
 
-	// Add Function and ColorMode properties at the beginning
-	TSharedPtr<IPropertyHandle> ColorModePropertyHandle = DetailLayout->GetProperty(GET_MEMBER_NAME_CHECKED(UDMXPixelMappingFixtureGroupItemComponent, ColorMode));
-	OutputSettingsCategory.AddProperty(ColorModePropertyHandle);
-
-	// Register attributes
-	TSharedPtr<FFunctionAttribute> AttributeR = MakeShared<FFunctionAttribute>();
-	AttributeR->Handle = DetailLayout->GetProperty(GET_MEMBER_NAME_CHECKED(UDMXPixelMappingFixtureGroupItemComponent, AttributeR));
-	AttributeR->ExposeHandle = DetailLayout->GetProperty(GET_MEMBER_NAME_CHECKED(UDMXPixelMappingFixtureGroupItemComponent, AttributeRExpose));
-	AttributeR->InvertHandle = DetailLayout->GetProperty(GET_MEMBER_NAME_CHECKED(UDMXPixelMappingFixtureGroupItemComponent, AttributeRInvert));
-
-	TSharedPtr<FFunctionAttribute> AttributeG = MakeShared<FFunctionAttribute>();
-	AttributeG->Handle = DetailLayout->GetProperty(GET_MEMBER_NAME_CHECKED(UDMXPixelMappingFixtureGroupItemComponent, AttributeG));
-	AttributeG->ExposeHandle = DetailLayout->GetProperty(GET_MEMBER_NAME_CHECKED(UDMXPixelMappingFixtureGroupItemComponent, AttributeGExpose));
-	AttributeG->InvertHandle = DetailLayout->GetProperty(GET_MEMBER_NAME_CHECKED(UDMXPixelMappingFixtureGroupItemComponent, AttributeGInvert));
-
-	TSharedPtr<FFunctionAttribute> AttributeB = MakeShared<FFunctionAttribute>();
-	AttributeB->Handle = DetailLayout->GetProperty(GET_MEMBER_NAME_CHECKED(UDMXPixelMappingFixtureGroupItemComponent, AttributeB));
-	AttributeB->ExposeHandle = DetailLayout->GetProperty(GET_MEMBER_NAME_CHECKED(UDMXPixelMappingFixtureGroupItemComponent, AttributeBExpose));
-	AttributeB->InvertHandle = DetailLayout->GetProperty(GET_MEMBER_NAME_CHECKED(UDMXPixelMappingFixtureGroupItemComponent, AttributeBInvert));
-
-	RGBAttributes.Add(AttributeR);
-	RGBAttributes.Add(AttributeG);
-	RGBAttributes.Add(AttributeB);
-
-	// Register Monochrome attribute
-	TSharedPtr<FFunctionAttribute> MonochromeAttribute = MakeShared<FFunctionAttribute>();
-	MonochromeAttribute->Handle = DetailLayout->GetProperty(GET_MEMBER_NAME_CHECKED(UDMXPixelMappingFixtureGroupItemComponent, MonochromeIntensity));
-	MonochromeAttribute->ExposeHandle = DetailLayout->GetProperty(GET_MEMBER_NAME_CHECKED(UDMXPixelMappingFixtureGroupItemComponent, bMonochromeExpose));
-	MonochromeAttribute->InvertHandle = DetailLayout->GetProperty(GET_MEMBER_NAME_CHECKED(UDMXPixelMappingFixtureGroupItemComponent, bMonochromeInvert));
-	MonochromeAttributes.Add(MonochromeAttribute);
-
-	// Generate all RGB Expose and Invert rows
-	OutputSettingsCategory.AddCustomRow(FText::GetEmpty())
-		.Visibility(TAttribute<EVisibility>(this, &FDMXPixelMappingDetailCustomization_FixtureGroupItem::GetRGBAttributesVisibility))
-		.NameContent()
-		[
-			SNew(STextBlock).Text(LOCTEXT("ColorSample", "Color Sample"))
-		]
-		.ValueContent()
-		[
-			SAssignNew(ExposeAndInvertListView, SListView<TSharedPtr<FFunctionAttribute>>)
-			.ListItemsSource(&RGBAttributes)
-			.OnGenerateRow(this, &FDMXPixelMappingDetailCustomization_FixtureGroupItem::GenerateExposeAndInvertRow)
-		];
-
-	// Update RGB attributes
-	for (TSharedPtr<FFunctionAttribute>& Attribute : RGBAttributes)
-	{
-		DetailLayout->HideProperty(Attribute->ExposeHandle);
-		DetailLayout->HideProperty(Attribute->InvertHandle);
-
-		OutputSettingsCategory
-			.AddProperty(Attribute->Handle)
-			.Visibility(TAttribute<EVisibility>::Create(TAttribute<EVisibility>::FGetter::CreateSP(this, &FDMXPixelMappingDetailCustomization_FixtureGroupItem::GetRGBAttributeRowVisibilty, Attribute.Get())));
-	}
-
-	// Generate all Monochrome Expose and Invert rows
-	OutputSettingsCategory.AddCustomRow(FText::GetEmpty())
-		.Visibility(TAttribute<EVisibility>(this, &FDMXPixelMappingDetailCustomization_FixtureGroupItem::GetMonochromeAttributesVisibility))
-		.NameContent()
-		[
-			SNew(STextBlock).Text(LOCTEXT("ColorSample", "Color Sample"))
-		]
-		.ValueContent()
-		[
-			SAssignNew(ExposeAndInvertListView, SListView<TSharedPtr<FFunctionAttribute>>)
-			.ListItemsSource(&MonochromeAttributes)
-			.OnGenerateRow(this, &FDMXPixelMappingDetailCustomization_FixtureGroupItem::GenerateExposeAndInvertRow)
-		];
-
-	// Update Monochrome attributes
-	for (TSharedPtr<FFunctionAttribute>& Attribute : MonochromeAttributes)
-	{
-		DetailLayout->HideProperty(Attribute->ExposeHandle);
-		DetailLayout->HideProperty(Attribute->InvertHandle);
-
-		OutputSettingsCategory
-			.AddProperty(Attribute->Handle)
-			.Visibility(TAttribute<EVisibility>::Create(TAttribute<EVisibility>::FGetter::CreateSP(this, &FDMXPixelMappingDetailCustomization_FixtureGroupItem::GetMonochromeRowVisibilty, Attribute.Get())));
-	}
-
 	CreateModulatorDetails(InDetailLayout);
-}
 
-bool FDMXPixelMappingDetailCustomization_FixtureGroupItem::CheckComponentsDMXColorMode(const EDMXColorMode DMXColorMode) const
-{
-	for (TWeakObjectPtr<UDMXPixelMappingFixtureGroupItemComponent> ItemComponent : FixtureGroupItemComponents)
-	{
-		if (ItemComponent.IsValid() && ItemComponent->ColorMode == DMXColorMode)
+	// Sort categories
+	InDetailLayout.SortCategories([](const TMap<FName, IDetailCategoryBuilder*>& CategoryMap)
 		{
-			return true;
-		}
-	}
-
-	return false;
-}
-
-EVisibility FDMXPixelMappingDetailCustomization_FixtureGroupItem::GetRGBAttributeRowVisibilty(FFunctionAttribute* Attribute) const
-{
-	bool bIsVisible = false;
-
-	// 1. Check if current attribute is sampling now
-	FPropertyAccess::Result Result = Attribute->ExposeHandle->GetValue(bIsVisible);
-	if (Result == FPropertyAccess::Result::MultipleValues)
-	{
-		bIsVisible = true;
-	}
-
-	// 2. Check if current color mode is RGB
-	if (!CheckComponentsDMXColorMode(EDMXColorMode::CM_RGB))
-	{
-		bIsVisible = false;
-	}
-
-	return bIsVisible ? EVisibility::Visible : EVisibility::Collapsed;
-}
-
-EVisibility FDMXPixelMappingDetailCustomization_FixtureGroupItem::GetRGBAttributesVisibility() const
-{
-	return CheckComponentsDMXColorMode(EDMXColorMode::CM_RGB) ? EVisibility::Visible : EVisibility::Collapsed;
-}
-
-EVisibility FDMXPixelMappingDetailCustomization_FixtureGroupItem::GetMonochromeRowVisibilty(FFunctionAttribute* Attribute) const
-{
-	bool bIsVisible = false;
-
-	// 1. Check if current attribute is sampling now
-	FPropertyAccess::Result Result = Attribute->ExposeHandle->GetValue(bIsVisible);
-	if (Result == FPropertyAccess::Result::MultipleValues)
-	{
-		bIsVisible = true;
-	}
+			int32 MinSortOrder = TNumericLimits<int32>::Max();
+			for (const TPair<FName, IDetailCategoryBuilder*>& Pair : CategoryMap)
+			{
+				int32 SortOrder = Pair.Value->GetSortOrder();
+				MinSortOrder = FMath::Min(SortOrder, MinSortOrder);
+			}
 
-	// 2. Check if current color mode is Monochrome
-	if (!CheckComponentsDMXColorMode(EDMXColorMode::CM_Monochrome))
-	{
-		bIsVisible = false;
-	}
+			IDetailCategoryBuilder* const* ColorSpaceCategoryPtr = CategoryMap.Find("Color Space");
+			if (ColorSpaceCategoryPtr)
+			{
+				(*ColorSpaceCategoryPtr)->SetSortOrder(MinSortOrder - 3);
+			}
 
-	return bIsVisible ? EVisibility::Visible : EVisibility::Collapsed;
-}
+			// Either 'RGB' 'XY' or 'XYZ' is displayed
+			IDetailCategoryBuilder* const* RGBCategoryPtr = CategoryMap.Find("RGB");
+			if (RGBCategoryPtr)
+			{
+				(*RGBCategoryPtr)->SetSortOrder(MinSortOrder - 2);
+			}
 
-EVisibility FDMXPixelMappingDetailCustomization_FixtureGroupItem::GetMonochromeAttributesVisibility() const
-{
-	return (GetRGBAttributesVisibility() == EVisibility::Visible) ? EVisibility::Collapsed : EVisibility::Visible;
-}
+			IDetailCategoryBuilder* const* XYCategoryPtr = CategoryMap.Find("XY");
+			if (XYCategoryPtr)
+			{
+				(*XYCategoryPtr)->SetSortOrder(MinSortOrder - 2);
+			}
 
-TSharedRef<ITableRow> FDMXPixelMappingDetailCustomization_FixtureGroupItem::GenerateExposeAndInvertRow(TSharedPtr<FFunctionAttribute> InAttribute, const TSharedRef<STableViewBase>& OwnerTable)
-{
-	if (!InAttribute.IsValid())
-	{
-		return SNew(STableRow<TSharedPtr<FString>>, OwnerTable);
-	}
+			IDetailCategoryBuilder* const* XYZCategoryPtr = CategoryMap.Find("XYZ");
+			if (XYZCategoryPtr)
+			{
+				(*XYZCategoryPtr)->SetSortOrder(MinSortOrder - 2);
+			}
 
-	return SNew(STableRow<TSharedPtr<FString>>, OwnerTable)
-		.Padding(2.0f)
-		.ShowSelection(false)
-		[
-			SNew(SBox)
-			[
-				SNew(SHorizontalBox)
-				+ SHorizontalBox::Slot()
-				.FillWidth(1)
-				.AutoWidth()
-				.Padding(2.0f, 0.0f)
-				.HAlign(HAlign_Left)
-				[
-					InAttribute->ExposeHandle->CreatePropertyNameWidget()
-				]
-				+ SHorizontalBox::Slot()
-				.FillWidth(1)
-				.AutoWidth()
-				.Padding(2.0f, 0.0f)
-				.HAlign(HAlign_Left)
-				[
-					InAttribute->ExposeHandle->CreatePropertyValueWidget()
-				]
-				+ SHorizontalBox::Slot()
-				.FillWidth(1)
-				.AutoWidth()
-				.Padding(2.0f, 0.0f)
-				.HAlign(HAlign_Left)
-				[
-					InAttribute->InvertHandle->CreatePropertyNameWidget()
-				]
-				+ SHorizontalBox::Slot()
-				.FillWidth(1)
-				.AutoWidth()
-				.Padding(2.0f, 0.0f)
-				.HAlign(HAlign_Left)
-				[
-					InAttribute->InvertHandle->CreatePropertyValueWidget()
-				]
-			]
-		];
+			IDetailCategoryBuilder* const* IntensityCategoryPtr = CategoryMap.Find("Luminance");
+			if (IntensityCategoryPtr)
+			{
+				(*IntensityCategoryPtr)->SetSortOrder(MinSortOrder - 1);
+			}
+		});
 }
 
 void FDMXPixelMappingDetailCustomization_FixtureGroupItem::CreateModulatorDetails(IDetailLayoutBuilder& InDetailLayout)
 {
 	IDetailCategoryBuilder& ModualtorsCategory = InDetailLayout.EditCategory("Modulators", LOCTEXT("DMXModulatorsCategory", "Modulators"), ECategoryPriority::Important);
 
-	TSharedPtr<IPropertyHandle> ModulatorClassesHandle = DetailLayout->GetProperty(GET_MEMBER_NAME_CHECKED(UDMXPixelMappingFixtureGroupItemComponent, ModulatorClasses), UDMXPixelMappingFixtureGroupItemComponent::StaticClass());
+	TSharedPtr<IPropertyHandle> ModulatorClassesHandle = InDetailLayout.GetProperty(GET_MEMBER_NAME_CHECKED(UDMXPixelMappingFixtureGroupItemComponent, ModulatorClasses), UDMXPixelMappingFixtureGroupItemComponent::StaticClass());
 	ModulatorClassesHandle->SetOnPropertyValueChanged(FSimpleDelegate::CreateSP(this, &FDMXPixelMappingDetailCustomization_FixtureGroupItem::ForceRefresh));
 	ModulatorClassesHandle->SetOnChildPropertyValueChanged(FSimpleDelegate::CreateSP(this, &FDMXPixelMappingDetailCustomization_FixtureGroupItem::ForceRefresh));
 
 	ModualtorsCategory.AddProperty(ModulatorClassesHandle);
 
-	TSharedPtr<IPropertyHandle> ModulatorsHandle = DetailLayout->GetProperty(GET_MEMBER_NAME_CHECKED(UDMXPixelMappingFixtureGroupItemComponent, Modulators), UDMXPixelMappingFixtureGroupItemComponent::StaticClass());
+	TSharedPtr<IPropertyHandle> ModulatorsHandle = InDetailLayout.GetProperty(GET_MEMBER_NAME_CHECKED(UDMXPixelMappingFixtureGroupItemComponent, Modulators), UDMXPixelMappingFixtureGroupItemComponent::StaticClass());
 	InDetailLayout.HideProperty(ModulatorsHandle);
 
 	// Create detail views for the modulators
diff --git a/Engine/Plugins/VirtualProduction/DMX/DMXPixelMapping/Source/DMXPixelMappingEditor/Private/Customizations/DMXPixelMappingDetailCustomization_FixtureGroupItem.h b/Engine/Plugins/VirtualProduction/DMX/DMXPixelMapping/Source/DMXPixelMappingEditor/Private/Customizations/DMXPixelMappingDetailCustomization_FixtureGroupItem.h
--- a/Engine/Plugins/VirtualProduction/DMX/DMXPixelMapping/Source/DMXPixelMappingEditor/Private/Customizations/DMXPixelMappingDetailCustomization_FixtureGroupItem.h
+++ b/Engine/Plugins/VirtualProduction/DMX/DMXPixelMapping/Source/DMXPixelMappingEditor/Private/Customizations/DMXPixelMappingDetailCustomization_FixtureGroupItem.h
@@ -3,17 +3,15 @@
 #pragma once
 
 #include "IDetailCustomization.h"
+#include "Layout/Visibility.h"
 #include "Misc/Attribute.h"
-#include "UObject/WeakObjectPtr.h" // IWYU pragma: keep
+#include "UObject/WeakObjectPtr.h"
 
-#include "Layout/Visibility.h"
 
 class IPropertyUtilities;
 
 class FDMXPixelMappingToolkit;
 class IDetailLayoutBuilder;
-class ITableRow;
-class STableViewBase;
 class IPropertyHandle;
 class UDMXPixelMappingFixtureGroupItemComponent;
 enum class EDMXColorMode : uint8;
@@ -51,38 +49,16 @@ public:
 	//~ IPropertyTypeCustomization interface end
 
 private:
-
-	EVisibility GetRGBAttributeRowVisibilty(FFunctionAttribute* Attribute) const;
-
-	EVisibility GetRGBAttributesVisibility() const;
-
-	EVisibility GetMonochromeRowVisibilty(FFunctionAttribute* Attribute) const;
-
-	EVisibility GetMonochromeAttributesVisibility() const;
-
-	TSharedRef<ITableRow> GenerateExposeAndInvertRow(TSharedPtr<FFunctionAttribute> InAttribute, const TSharedRef<STableViewBase>& OwnerTable);
-
-private:
-	bool CheckComponentsDMXColorMode(const EDMXColorMode DMXColorMode) const;
-
 	/** Creates Details for the Output Modulators */
 	void CreateModulatorDetails(IDetailLayoutBuilder& InDetailLayout);
 
 	/** Forces the layout to redraw */
 	void ForceRefresh();
 
-	IDetailLayoutBuilder* DetailLayout;
-
+	/** Fixture Group Item Components that are being edited */
 	TArray<TWeakObjectPtr<UDMXPixelMappingFixtureGroupItemComponent>> FixtureGroupItemComponents;
 
-	TArray<TSharedPtr<FName>> ActiveModeFunctions;
-
-	TArray<TSharedPtr<FFunctionAttribute>> RGBAttributes;
-
-	TArray<TSharedPtr<FFunctionAttribute>> MonochromeAttributes;
-
-	TSharedPtr<SListView<TSharedPtr<FFunctionAttribute>>> ExposeAndInvertListView;
-
+	/** Property utilities for this customization */
 	TSharedPtr<IPropertyUtilities> PropertyUtilities;
 
 	/** Weak reference to the DMX editor */
diff --git a/Engine/Plugins/VirtualProduction/DMX/DMXPixelMapping/Source/DMXPixelMappingEditor/Private/Customizations/DMXPixelMappingDetailCustomization_Matrix.cpp b/Engine/Plugins/VirtualProduction/DMX/DMXPixelMapping/Source/DMXPixelMappingEditor/Private/Customizations/DMXPixelMappingDetailCustomization_Matrix.cpp
--- a/Engine/Plugins/VirtualProduction/DMX/DMXPixelMapping/Source/DMXPixelMappingEditor/Private/Customizations/DMXPixelMappingDetailCustomization_Matrix.cpp
+++ b/Engine/Plugins/VirtualProduction/DMX/DMXPixelMapping/Source/DMXPixelMappingEditor/Private/Customizations/DMXPixelMappingDetailCustomization_Matrix.cpp
@@ -4,16 +4,15 @@
 
 #include "DMXPixelMappingTypes.h"
 #include "Components/DMXPixelMappingMatrixComponent.h"
-#include "DetailsViewArgs.h"
-#include "Modules/ModuleManager.h"
+#include "Modulators/DMXModulator.h"
 
-#include "DetailWidgetRow.h"
 #include "DetailLayoutBuilder.h"
-#include "Framework/Views/TableViewMetadata.h"
+#include "DetailsViewArgs.h"
+#include "DetailWidgetRow.h"
 #include "IPropertyUtilities.h"
 #include "IDetailsView.h"
+#include "Framework/Views/TableViewMetadata.h"
 #include "Modules/ModuleManager.h"
-#include "Modulators/DMXModulator.h"
 #include "Widgets/Views/SListView.h"
 
 
@@ -21,12 +20,11 @@
 
 void FDMXPixelMappingDetailCustomization_Matrix::CustomizeDetails(IDetailLayoutBuilder& InDetailLayout)
 {
-	DetailLayout = &InDetailLayout;
 	PropertyUtilities = InDetailLayout.GetPropertyUtilities();
 
 	// Get editing UObject
 	TArray<TWeakObjectPtr<UObject>> OuterObjects;
-	DetailLayout->GetObjectsBeingCustomized(OuterObjects);
+	InDetailLayout.GetObjectsBeingCustomized(OuterObjects);
 
 	MatrixComponents.Empty();
 
@@ -36,220 +34,25 @@ void FDMXPixelMappingDetailCustomization_Matrix::CustomizeDetails(IDetailLayoutB
 	}
 
 	// Get editing categories
-	IDetailCategoryBuilder& OutputSettingsCategory = DetailLayout->EditCategory("Output Settings", FText::GetEmpty(), ECategoryPriority::Important);
+	IDetailCategoryBuilder& OutputSettingsCategory = InDetailLayout.EditCategory("Output Settings", FText::GetEmpty(), ECategoryPriority::Important);
 
 	// Hide the Layout Script property (shown in its own panel, see SDMXPixelMappingLayoutView)
 	InDetailLayout.HideProperty(GET_MEMBER_NAME_CHECKED(UDMXPixelMappingMatrixComponent, LayoutScript));
 
-	// Add color mode property
-	ColorModePropertyHandle = DetailLayout->GetProperty(GET_MEMBER_NAME_CHECKED(UDMXPixelMappingMatrixComponent, ColorMode), UDMXPixelMappingMatrixComponent::StaticClass());
-	OutputSettingsCategory.AddProperty(ColorModePropertyHandle);
-
-	// Register attributes
-	TSharedPtr<FDMXCellAttributeGroup> AttributeR = MakeShared<FDMXCellAttributeGroup>();
-	AttributeR->Handle = DetailLayout->GetProperty(GET_MEMBER_NAME_CHECKED(UDMXPixelMappingMatrixComponent, AttributeR));
-	AttributeR->ExposeHandle = DetailLayout->GetProperty(GET_MEMBER_NAME_CHECKED(UDMXPixelMappingMatrixComponent, AttributeRExpose));
-	AttributeR->InvertHandle = DetailLayout->GetProperty(GET_MEMBER_NAME_CHECKED(UDMXPixelMappingMatrixComponent, AttributeRInvert));
-
-	TSharedPtr<FDMXCellAttributeGroup> AttributeG = MakeShared<FDMXCellAttributeGroup>();
-	AttributeG->Handle = DetailLayout->GetProperty(GET_MEMBER_NAME_CHECKED(UDMXPixelMappingMatrixComponent, AttributeG));
-	AttributeG->ExposeHandle = DetailLayout->GetProperty(GET_MEMBER_NAME_CHECKED(UDMXPixelMappingMatrixComponent, AttributeGExpose));
-	AttributeG->InvertHandle = DetailLayout->GetProperty(GET_MEMBER_NAME_CHECKED(UDMXPixelMappingMatrixComponent, AttributeGInvert));
-
-	TSharedPtr<FDMXCellAttributeGroup> AttributeB = MakeShared<FDMXCellAttributeGroup>();
-	AttributeB->Handle = DetailLayout->GetProperty(GET_MEMBER_NAME_CHECKED(UDMXPixelMappingMatrixComponent, AttributeB));
-	AttributeB->ExposeHandle = DetailLayout->GetProperty(GET_MEMBER_NAME_CHECKED(UDMXPixelMappingMatrixComponent, AttributeBExpose));
-	AttributeB->InvertHandle = DetailLayout->GetProperty(GET_MEMBER_NAME_CHECKED(UDMXPixelMappingMatrixComponent, AttributeBInvert));
-
-	RGBAttributes.Add(AttributeR);
-	RGBAttributes.Add(AttributeG);
-	RGBAttributes.Add(AttributeB);
-
-	// Register Monochrome attribute
-	TSharedPtr<FDMXCellAttributeGroup> MonochromeAttribute = MakeShared<FDMXCellAttributeGroup>();
-	MonochromeAttribute->Handle = DetailLayout->GetProperty(GET_MEMBER_NAME_CHECKED(UDMXPixelMappingMatrixComponent, MonochromeIntensity));
-	MonochromeAttribute->ExposeHandle = DetailLayout->GetProperty(GET_MEMBER_NAME_CHECKED(UDMXPixelMappingMatrixComponent, bMonochromeExpose));
-	MonochromeAttribute->InvertHandle = DetailLayout->GetProperty(GET_MEMBER_NAME_CHECKED(UDMXPixelMappingMatrixComponent, bMonochromeInvert));
-	MonochromeAttributes.Add(MonochromeAttribute);
-
-	// Generate all RGB Expose and Invert rows
-	OutputSettingsCategory.AddCustomRow(FText::GetEmpty())
-		.Visibility(TAttribute<EVisibility>(this, &FDMXPixelMappingDetailCustomization_Matrix::GetRGBAttributesVisibility))
-		.NameContent()
-		[
-			SNew(STextBlock).Text(LOCTEXT("ColorSample", "Color Sample"))
-		]
-		.ValueContent()
-		[
-			SAssignNew(ExposeAndInvertListView, SListView<TSharedPtr<FDMXCellAttributeGroup>>)
-			.ListItemsSource(&RGBAttributes)
-			.OnGenerateRow(this, &FDMXPixelMappingDetailCustomization_Matrix::GenerateExposeAndInvertRow)
-		];
-
-	// Update RGB attributes
-	for (TSharedPtr<FDMXCellAttributeGroup>& Attribute : RGBAttributes)
-	{
-		DetailLayout->HideProperty(Attribute->ExposeHandle);
-		DetailLayout->HideProperty(Attribute->InvertHandle);
-
-		OutputSettingsCategory
-			.AddProperty(Attribute->Handle)
-			.Visibility(TAttribute<EVisibility>::Create(TAttribute<EVisibility>::FGetter::CreateSP(this, &FDMXPixelMappingDetailCustomization_Matrix::GetRGBAttributeRowVisibilty, Attribute.Get())));
-	}
-
-
-	// Generate all Monochrome Expose and Invert rows
-	OutputSettingsCategory.AddCustomRow(FText::GetEmpty())
-		.Visibility(TAttribute<EVisibility>(this, &FDMXPixelMappingDetailCustomization_Matrix::GetMonochromeAttributesVisibility))
-		.NameContent()
-		[
-			SNew(STextBlock).Text(LOCTEXT("ColorSample", "Color Sample"))
-		]
-		.ValueContent()
-		[
-			SAssignNew(ExposeAndInvertListView, SListView<TSharedPtr<FDMXCellAttributeGroup>>)
-			.ListItemsSource(&MonochromeAttributes)
-			.OnGenerateRow(this, &FDMXPixelMappingDetailCustomization_Matrix::GenerateExposeAndInvertRow)
-		];
-
-	// Update Monochrome attributes
-	for (TSharedPtr<FDMXCellAttributeGroup>& Attribute : MonochromeAttributes)
-	{
-		DetailLayout->HideProperty(Attribute->ExposeHandle);
-		DetailLayout->HideProperty(Attribute->InvertHandle);
-
-		OutputSettingsCategory
-			.AddProperty(Attribute->Handle)
-			.Visibility(TAttribute<EVisibility>::Create(TAttribute<EVisibility>::FGetter::CreateSP(this, &FDMXPixelMappingDetailCustomization_Matrix::GetMonochromeRowVisibilty, Attribute.Get())));
-	}
-	
 	CreateModulatorDetails(InDetailLayout);
 }
 
-bool FDMXPixelMappingDetailCustomization_Matrix::CheckComponentsDMXColorMode(const EDMXColorMode DMXColorMode) const
-{
-	for (TWeakObjectPtr<UDMXPixelMappingMatrixComponent> ItemComponent : MatrixComponents)
-	{
-		if (ItemComponent.IsValid() && ItemComponent->ColorMode == DMXColorMode)
-		{
-			return true;
-		}
-	}
-
-	return false;
-}
-
-EVisibility FDMXPixelMappingDetailCustomization_Matrix::GetRGBAttributeRowVisibilty(FDMXCellAttributeGroup* Attribute) const
-{
-	bool bIsVisible = false;
-
-	// 1. Check if current attribute is sampling now
-	FPropertyAccess::Result Result = Attribute->ExposeHandle->GetValue(bIsVisible);
-	if (Result == FPropertyAccess::Result::MultipleValues)
-	{
-		bIsVisible = true;
-	}
-
-	// 2. Check if current color mode is RGB
-	if (!CheckComponentsDMXColorMode(EDMXColorMode::CM_RGB))
-	{
-		bIsVisible = false;
-	}
-
-	return bIsVisible ? EVisibility::Visible : EVisibility::Collapsed;
-}
-
-EVisibility FDMXPixelMappingDetailCustomization_Matrix::GetRGBAttributesVisibility() const
-{	
-	return CheckComponentsDMXColorMode(EDMXColorMode::CM_RGB) ? EVisibility::Visible : EVisibility::Collapsed;
-}
-
-EVisibility FDMXPixelMappingDetailCustomization_Matrix::GetMonochromeRowVisibilty(FDMXCellAttributeGroup* Attribute) const
-{
-	bool bIsVisible = false;
-
-	// 1. Check if current attribute is sampling now
-	FPropertyAccess::Result Result = Attribute->ExposeHandle->GetValue(bIsVisible);
-	if (Result == FPropertyAccess::Result::MultipleValues)
-	{
-		bIsVisible = true;
-	}
-
-	// 2. Check if current color mode is Monochrome
-	if (!CheckComponentsDMXColorMode(EDMXColorMode::CM_Monochrome))
-	{
-		bIsVisible = false;
-	}
-
-	return bIsVisible ? EVisibility::Visible : EVisibility::Collapsed;
-}
-
-EVisibility FDMXPixelMappingDetailCustomization_Matrix::GetMonochromeAttributesVisibility() const
-{
-	return (GetRGBAttributesVisibility() == EVisibility::Visible) ? EVisibility::Collapsed : EVisibility::Visible;
-}
-
-TSharedRef<ITableRow> FDMXPixelMappingDetailCustomization_Matrix::GenerateExposeAndInvertRow(TSharedPtr<FDMXCellAttributeGroup> InAtribute, const TSharedRef<STableViewBase>& OwnerTable)
-{
-	if (!InAtribute.IsValid())
-	{
-		return SNew(STableRow<TSharedPtr<FString>>, OwnerTable);
-	}
-
-	return SNew(STableRow<TSharedPtr<FString>>, OwnerTable)
-		.Padding(2.0f)
-		.ShowSelection(false)
-		[
-			SNew(SBox)
-			[
-				SNew(SHorizontalBox)
-				+ SHorizontalBox::Slot()
-				.FillWidth(1)
-				.AutoWidth()
-				.Padding(2.0f, 0.0f)
-				.HAlign(HAlign_Left)
-				[
-					InAtribute->ExposeHandle->CreatePropertyNameWidget()
-				]
-				+ SHorizontalBox::Slot()
-				.FillWidth(1)
-				.AutoWidth()
-				.Padding(2.0f, 0.0f)
-				.HAlign(HAlign_Left)
-				[
-					InAtribute->ExposeHandle->CreatePropertyValueWidget()
-				]
-				+ SHorizontalBox::Slot()
-				.FillWidth(1)
-				.AutoWidth()
-				.Padding(2.0f, 0.0f)
-				.HAlign(HAlign_Left)
-				[
-					InAtribute->InvertHandle->CreatePropertyNameWidget()
-				]
-				+ SHorizontalBox::Slot()
-				.FillWidth(1)
-				.AutoWidth()
-				.Padding(2.0f, 0.0f)
-				.HAlign(HAlign_Left)
-				[
-					InAtribute->InvertHandle->CreatePropertyValueWidget()
-				]
-			]
-		];
-}
-
 void FDMXPixelMappingDetailCustomization_Matrix::CreateModulatorDetails(IDetailLayoutBuilder& InDetailLayout)
 {
 	IDetailCategoryBuilder& ModualtorsCategory = InDetailLayout.EditCategory("Modulators", LOCTEXT("DMXModulatorsCategory", "Modulators"), ECategoryPriority::Important);
 
-	TSharedPtr<IPropertyHandle> ModulatorClassesHandle = DetailLayout->GetProperty(GET_MEMBER_NAME_CHECKED(UDMXPixelMappingMatrixComponent, ModulatorClasses), UDMXPixelMappingMatrixComponent::StaticClass());
+	TSharedPtr<IPropertyHandle> ModulatorClassesHandle = InDetailLayout.GetProperty(GET_MEMBER_NAME_CHECKED(UDMXPixelMappingMatrixComponent, ModulatorClasses), UDMXPixelMappingMatrixComponent::StaticClass());
 	ModulatorClassesHandle->SetOnPropertyValueChanged(FSimpleDelegate::CreateSP(this, &FDMXPixelMappingDetailCustomization_Matrix::ForceRefresh));
 	ModulatorClassesHandle->SetOnChildPropertyValueChanged(FSimpleDelegate::CreateSP(this, &FDMXPixelMappingDetailCustomization_Matrix::ForceRefresh));
 
 	ModualtorsCategory.AddProperty(ModulatorClassesHandle);
 
-	TSharedPtr<IPropertyHandle> ModulatorsHandle = DetailLayout->GetProperty(GET_MEMBER_NAME_CHECKED(UDMXPixelMappingMatrixComponent, Modulators), UDMXPixelMappingMatrixComponent::StaticClass());
+	TSharedPtr<IPropertyHandle> ModulatorsHandle = InDetailLayout.GetProperty(GET_MEMBER_NAME_CHECKED(UDMXPixelMappingMatrixComponent, Modulators), UDMXPixelMappingMatrixComponent::StaticClass());
 	InDetailLayout.HideProperty(ModulatorsHandle);
 
 	// Create detail views for the modulators
diff --git a/Engine/Plugins/VirtualProduction/DMX/DMXPixelMapping/Source/DMXPixelMappingEditor/Private/Customizations/DMXPixelMappingDetailCustomization_Matrix.h b/Engine/Plugins/VirtualProduction/DMX/DMXPixelMapping/Source/DMXPixelMappingEditor/Private/Customizations/DMXPixelMappingDetailCustomization_Matrix.h
--- a/Engine/Plugins/VirtualProduction/DMX/DMXPixelMapping/Source/DMXPixelMappingEditor/Private/Customizations/DMXPixelMappingDetailCustomization_Matrix.h
+++ b/Engine/Plugins/VirtualProduction/DMX/DMXPixelMapping/Source/DMXPixelMappingEditor/Private/Customizations/DMXPixelMappingDetailCustomization_Matrix.h
@@ -4,8 +4,8 @@
 
 #include "IDetailCustomization.h"
 #include "Misc/Attribute.h"
-#include "UObject/WeakObjectPtr.h" // IWYU pragma: keep
 #include "Layout/Visibility.h"
+#include "UObject/WeakObjectPtr.h"
 
 class FDMXPixelMappingToolkit;
 struct FDMXCellAttributeGroup;
@@ -46,41 +46,17 @@ public:
 	/** IDetailCustomization interface */
 	virtual void CustomizeDetails(IDetailLayoutBuilder& DetailLayout) override;
 
-private:
-	EVisibility GetRGBAttributeRowVisibilty(FDMXCellAttributeGroup* Attribute) const;
-
-	EVisibility GetRGBAttributesVisibility() const;
-
-	EVisibility GetMonochromeRowVisibilty(FDMXCellAttributeGroup* Attribute) const;
-
-	EVisibility GetMonochromeAttributesVisibility() const;
-
-	TSharedRef<ITableRow> GenerateExposeAndInvertRow(TSharedPtr<FDMXCellAttributeGroup> InAtribute, const TSharedRef<STableViewBase>& OwnerTable);
-
-
 private:	
-	bool CheckComponentsDMXColorMode(const EDMXColorMode DMXColorMode) const;
-
 	/** Creates Details for the Output Modulators */
 	void CreateModulatorDetails(IDetailLayoutBuilder& InDetailLayout);
 
 	/** Forces the layout to redraw */
 	void ForceRefresh();
 
+	/** Matrix Components that are being edited */
 	TArray<TWeakObjectPtr<UDMXPixelMappingMatrixComponent>> MatrixComponents;
 
-	IDetailLayoutBuilder* DetailLayout;
-
-	TArray<TSharedPtr<FName>> ActiveModeFunctions;
-
-	TArray<TSharedPtr<FDMXCellAttributeGroup>> RGBAttributes;
-
-	TArray<TSharedPtr<FDMXCellAttributeGroup>> MonochromeAttributes;
-
-	TSharedPtr<IPropertyHandle> ColorModePropertyHandle;
-
-	TSharedPtr<SListView<TSharedPtr<FDMXCellAttributeGroup>>> ExposeAndInvertListView;
-
+	/** Property utilities for this customization */
 	TSharedPtr<IPropertyUtilities> PropertyUtilities;
 
 	/** Weak reference to the DMX editor */
diff --git a/Engine/Plugins/VirtualProduction/DMX/DMXPixelMapping/Source/DMXPixelMappingEditor/Private/Toolkits/DMXPixelMappingToolkit.cpp b/Engine/Plugins/VirtualProduction/DMX/DMXPixelMapping/Source/DMXPixelMappingEditor/Private/Toolkits/DMXPixelMappingToolkit.cpp
--- a/Engine/Plugins/VirtualProduction/DMX/DMXPixelMapping/Source/DMXPixelMappingEditor/Private/Toolkits/DMXPixelMappingToolkit.cpp
+++ b/Engine/Plugins/VirtualProduction/DMX/DMXPixelMapping/Source/DMXPixelMappingEditor/Private/Toolkits/DMXPixelMappingToolkit.cpp
@@ -3,12 +3,12 @@
 #include "Toolkits/DMXPixelMappingToolkit.h"
 
 #include "DMXEditorUtils.h"
-#include "DMXPixelMappingLayoutSettings.h"
 #include "DMXPixelMapping.h"
 #include "DMXPixelMappingEditorCommands.h"
 #include "DMXPixelMappingEditorModule.h"
 #include "DMXPixelMappingEditorStyle.h"
 #include "DMXPixelMappingEditorUtils.h"
+#include "DMXPixelMappingLayoutSettings.h"
 #include "DMXPixelMappingToolbar.h"
 #include "K2Node_PixelMappingBaseComponent.h"
 #include "Components/DMXPixelMappingFixtureGroupComponent.h"
@@ -16,10 +16,7 @@
 #include "Components/DMXPixelMappingRootComponent.h"
 #include "Components/DMXPixelMappingMatrixComponent.h"
 #include "Components/DMXPixelMappingScreenComponent.h"
-#include "Framework/MultiBox/MultiBoxExtender.h"
 #include "Templates/DMXPixelMappingComponentTemplate.h"
-#include "Modules/ModuleManager.h"
-#include "UObject/UObjectIterator.h"
 #include "ViewModels/DMXPixelMappingPaletteViewModel.h"
 #include "Views/SDMXPixelMappingDesignerView.h"
 #include "Views/SDMXPixelMappingDetailsView.h"
@@ -31,6 +28,9 @@
 #include "ScopedTransaction.h"
 #include "Engine/TextureRenderTarget2D.h"
 #include "Framework/Commands/GenericCommands.h"
+#include "Framework/MultiBox/MultiBoxExtender.h"
+#include "Modules/ModuleManager.h"
+#include "UObject/UObjectIterator.h"
 #include "Widgets/Docking/SDockTab.h"
 
 
@@ -43,13 +43,9 @@ const FName FDMXPixelMappingToolkit::PreviewViewTabID(TEXT("DMXPixelMappingEdito
 const FName FDMXPixelMappingToolkit::DetailsViewTabID(TEXT("DMXPixelMappingEditor_DetailsViewTabID"));
 const FName FDMXPixelMappingToolkit::LayoutViewTabID(TEXT("DMXPixelMappingEditor_LayoutViewTabID"));
 
-const uint8 FDMXPixelMappingToolkit::RequestStopSendingMaxTicks = 5;
 
 FDMXPixelMappingToolkit::FDMXPixelMappingToolkit()
 	: DMXPixelMapping(nullptr)
-	, bIsPlayingDMX(false)
-	, bTogglePlayDMXAll(true)
-	, bRequestStopSendingDMX(false)
 {}
 
 FDMXPixelMappingToolkit::~FDMXPixelMappingToolkit()
@@ -149,7 +145,7 @@ void FDMXPixelMappingToolkit::Tick(float DeltaTime)
 	}
 
 	// render selected component
-	if (!bIsPlayingDMX)
+	if (!bIsPlayingDMX && !SelectedComponents.IsEmpty())
 	{
 		for (const FDMXPixelMappingComponentReference& SelectedComponentRef : SelectedComponents)
 		{
@@ -181,49 +177,26 @@ void FDMXPixelMappingToolkit::Tick(float DeltaTime)
 		}
 	}
 
-	if (bIsPlayingDMX)
+	if (bIsPlayingDMX && !bRequestStopSendingDMX)
 	{
 		if (bTogglePlayDMXAll) // Send to all
 		{
 			// Render all components
 			RootComponent->RenderAndSendDMX();
-
-			for (FDMXPixelMappingComponentReference& SelectedComponentRef : SelectedComponents)
-			{
-				if (UDMXPixelMappingBaseComponent* SelectedComponent = SelectedComponentRef.Component.Get())
-				{
-					// Try to get renderer component from selected component
-					UDMXPixelMappingRendererComponent* RendererComponent = SelectedComponent->GetRendererComponent();
-					if (RendererComponent)
-					{
-						RendererComponent->RenderEditorPreviewTexture();
-					}
-
-					// Render only once for all selected components
-					break;
-				}
-			}
 		}
 		else // Send to selected component
 		{
-			bool bRenderedOnce = false;
-
 			for (FDMXPixelMappingComponentReference& SelectedComponentRef : SelectedComponents)
 			{
+				TSet<UDMXPixelMappingRendererComponent*> RenderedRendererComponents;
 				if (UDMXPixelMappingBaseComponent* SelectedComponent = SelectedComponentRef.Component.Get())
 				{
-					if (!bRenderedOnce)
+					// Try to get renderer component from selected component
+					UDMXPixelMappingRendererComponent* RendererComponent = SelectedComponent->GetRendererComponent();
+					if (RendererComponent && !RenderedRendererComponents.Contains(RendererComponent))
 					{
-						// Try to get renderer component from selected component
-						UDMXPixelMappingRendererComponent* RendererComponent = SelectedComponent->GetRendererComponent();
-						if (RendererComponent)
-						{
-							RendererComponent->Render();
-
-							RendererComponent->RenderEditorPreviewTexture();
-						}
-
-						bRenderedOnce = true;
+						RendererComponent->Render();
+						RenderedRendererComponents.Add(RendererComponent);
 					}
 
 					SelectedComponent->SendDMX();
@@ -231,35 +204,12 @@ void FDMXPixelMappingToolkit::Tick(float DeltaTime)
 			}
 		}
 	}
-	else
+	else if (bRequestStopSendingDMX)
 	{
-		if (bRequestStopSendingDMX)
-		{
-			if (RequestStopSendingTicks < RequestStopSendingMaxTicks)
-			{
-				if (bTogglePlayDMXAll) // Send to all
-				{
-					RootComponent->ResetDMX();
-				}
-				else // Send to selected component
-				{
-					for (FDMXPixelMappingComponentReference& SelectedComponentRef : SelectedComponents)
-					{
-						if (UDMXPixelMappingBaseComponent* SelectedComponent = SelectedComponentRef.Component.Get())
-						{
-							SelectedComponent->ResetDMX();
-						}
-					}
-				}
-
-				RequestStopSendingTicks++;
-			}
-			else
-			{
-				RequestStopSendingTicks = 0;
-				bRequestStopSendingDMX = false;
-			}
-		}
+		RootComponent->ResetDMX();
+			
+		bIsPlayingDMX = false;
+		bRequestStopSendingDMX = false; 
 	}
 }
 
@@ -355,13 +305,7 @@ void FDMXPixelMappingToolkit::PlayDMX()
 
 void FDMXPixelMappingToolkit::StopPlayingDMX()
 {
-	bIsPlayingDMX = false;
-
-	RequestStopSendingTicks = 0;
 	bRequestStopSendingDMX = true;
-
-	FDMXEditorUtils::ClearFixturePatchCachedData();
-	FDMXEditorUtils::ClearAllDMXPortBuffers();
 }
 
 void FDMXPixelMappingToolkit::ExecutebTogglePlayDMXAll()
diff --git a/Engine/Plugins/VirtualProduction/DMX/DMXPixelMapping/Source/DMXPixelMappingEditor/Private/Toolkits/DMXPixelMappingToolkit.h b/Engine/Plugins/VirtualProduction/DMX/DMXPixelMapping/Source/DMXPixelMappingEditor/Private/Toolkits/DMXPixelMappingToolkit.h
--- a/Engine/Plugins/VirtualProduction/DMX/DMXPixelMapping/Source/DMXPixelMappingEditor/Private/Toolkits/DMXPixelMappingToolkit.h
+++ b/Engine/Plugins/VirtualProduction/DMX/DMXPixelMapping/Source/DMXPixelMappingEditor/Private/Toolkits/DMXPixelMappingToolkit.h
@@ -2,9 +2,10 @@
 
 #pragma once
 
+#include "DMXPixelMappingComponentReference.h"
+
 #include "TickableEditorObject.h"
 #include "Toolkits/AssetEditorToolkit.h"
-#include "DMXPixelMappingComponentReference.h"
 
 class FDMXPixelMappingComponentTemplate;
 class FSpawnTabArgs;
@@ -241,15 +242,14 @@ private:
 	/** Command list for handling widget actions in the PixelMapping Toolkit */
 	TSharedPtr<FUICommandList> DesignerCommandList;
 
-	bool bIsPlayingDMX;
-
-	bool bTogglePlayDMXAll;
-
-	bool bRequestStopSendingDMX;
+	/** True while playing DMX */
+	bool bIsPlayingDMX = false;
 
-	uint8 RequestStopSendingTicks;
+	/** Toggles if DMX should be sent for all components */
+	bool bTogglePlayDMXAll = true;
 
-	static const uint8 RequestStopSendingMaxTicks;
+	/** True while stop DMX is requested, but not carried out yet */
+	bool bRequestStopSendingDMX = false;
 
 	/** True while adding components (to avoid needlessly updating blueprint nodes on each component added via our own methods) */
 	bool bAddingComponents = false;
diff --git a/Engine/Plugins/VirtualProduction/DMX/DMXPixelMapping/Source/DMXPixelMappingEditor/Private/Views/SDMXPixelMappingDetailsView.cpp b/Engine/Plugins/VirtualProduction/DMX/DMXPixelMapping/Source/DMXPixelMappingEditor/Private/Views/SDMXPixelMappingDetailsView.cpp
--- a/Engine/Plugins/VirtualProduction/DMX/DMXPixelMapping/Source/DMXPixelMappingEditor/Private/Views/SDMXPixelMappingDetailsView.cpp
+++ b/Engine/Plugins/VirtualProduction/DMX/DMXPixelMapping/Source/DMXPixelMappingEditor/Private/Views/SDMXPixelMappingDetailsView.cpp
@@ -2,8 +2,8 @@
 
 #include "Views/SDMXPixelMappingDetailsView.h"
 
-#include "DetailsViewArgs.h"
-#include "Toolkits/DMXPixelMappingToolkit.h"
+#include "DMXPixelMappingComponentReference.h"
+#include "ColorSpace/DMXPixelMappingColorSpace_RGBCMY.h"
 #include "Components/DMXPixelMappingRootComponent.h"
 #include "Components/DMXPixelMappingFixtureGroupComponent.h"
 #include "Components/DMXPixelMappingFixtureGroupItemComponent.h"
@@ -11,17 +11,22 @@
 #include "Components/DMXPixelMappingMatrixComponent.h"
 #include "Components/DMXPixelMappingMatrixCellComponent.h"
 #include "Components/DMXPixelMappingRendererComponent.h"
+#include "Customizations/DMXPixelMappingColorSpaceDetails_RGBCMY.h"
 #include "Customizations/DMXPixelMappingDetailCustomization_FixtureGroup.h"
 #include "Customizations/DMXPixelMappingDetailCustomization_FixtureGroupItem.h"
 #include "Customizations/DMXPixelMappingDetailCustomization_Screen.h"
 #include "Customizations/DMXPixelMappingDetailCustomization_Renderer.h"
 #include "Customizations/DMXPixelMappingDetailCustomization_Matrix.h"
 #include "Customizations/DMXPixelMappingDetailCustomization_MatrixCell.h"
+#include "Toolkits/DMXPixelMappingToolkit.h"
 
+#include "DetailsViewArgs.h"
 #include "IDetailsView.h"
+#include "PropertyEditorDelegates.h"
 #include "PropertyEditorModule.h"
 #include "Modules/ModuleManager.h"
 
+
 void SDMXPixelMappingDetailsView::Construct(const FArguments& InArgs, const TSharedPtr<FDMXPixelMappingToolkit>& InToolkit)
 {
 	ToolkitWeakPtr = InToolkit;
@@ -49,7 +54,7 @@ void SDMXPixelMappingDetailsView::Construct(const FArguments& InArgs, const TSha
 		]
 	];
 
-	InToolkit->GetOnSelectedComponentsChangedDelegate().AddRaw(this, &SDMXPixelMappingDetailsView::OnSelectedComponentsChanged);
+	InToolkit->GetOnSelectedComponentsChangedDelegate().AddSP(this, &SDMXPixelMappingDetailsView::OnSelectedComponentsChanged);
 
 	OnSelectedComponentsChanged();
 }
@@ -62,6 +67,10 @@ SDMXPixelMappingDetailsView::~SDMXPixelMappingDetailsView()
 		PropertyView->UnregisterInstancedCustomPropertyLayout(UDMXPixelMappingFixtureGroupItemComponent::StaticClass());
 		PropertyView->UnregisterInstancedCustomPropertyLayout(UDMXPixelMappingScreenComponent::StaticClass());
 		PropertyView->UnregisterInstancedCustomPropertyLayout(UDMXPixelMappingRendererComponent::StaticClass());
+		PropertyView->UnregisterInstancedCustomPropertyLayout(UDMXPixelMappingMatrixComponent::StaticClass());
+		PropertyView->UnregisterInstancedCustomPropertyLayout(UDMXPixelMappingMatrixCellComponent::StaticClass());
+		
+		PropertyView->UnregisterInstancedCustomPropertyLayout(UDMXPixelMappingColorSpace_RGBCMY::StaticClass());
 	}
 }
 
@@ -123,4 +132,7 @@ void SDMXPixelMappingDetailsView::RegisterCustomizations()
 
 	FOnGetDetailCustomizationInstance MatrixCellCustomizationInstance = FOnGetDetailCustomizationInstance::CreateStatic(&FDMXPixelMappingDetailCustomization_MatrixCell::MakeInstance, ToolkitWeakPtr);
 	PropertyView->RegisterInstancedCustomPropertyLayout(UDMXPixelMappingMatrixCellComponent::StaticClass(), MatrixCellCustomizationInstance);
+
+	FOnGetDetailCustomizationInstance ColorSpaceCustomizationInstance_RGBCMY = FOnGetDetailCustomizationInstance::CreateStatic(&FDMXPixelMappingColorSpaceDetails_RGBCMY::MakeInstance);
+	PropertyView->RegisterInstancedCustomPropertyLayout(UDMXPixelMappingColorSpace_RGBCMY::StaticClass(), ColorSpaceCustomizationInstance_RGBCMY);
 }
diff --git a/Engine/Plugins/VirtualProduction/DMX/DMXPixelMapping/Source/DMXPixelMappingRenderer/Private/DMXPixelMappingRenderer.cpp b/Engine/Plugins/VirtualProduction/DMX/DMXPixelMapping/Source/DMXPixelMappingRenderer/Private/DMXPixelMappingRenderer.cpp
--- a/Engine/Plugins/VirtualProduction/DMX/DMXPixelMapping/Source/DMXPixelMappingRenderer/Private/DMXPixelMappingRenderer.cpp
+++ b/Engine/Plugins/VirtualProduction/DMX/DMXPixelMapping/Source/DMXPixelMappingRenderer/Private/DMXPixelMappingRenderer.cpp
@@ -76,8 +76,6 @@ public:
 
 		SHADER_PARAMETER(FIntPoint, InputTextureSize)
 		SHADER_PARAMETER(FIntPoint, OutputTextureSize)
-		SHADER_PARAMETER(FVector4f, PixelFactor)
-		SHADER_PARAMETER(FIntVector4, InvertPixel)
 		SHADER_PARAMETER(FVector2f, UVCellSize)
 	END_SHADER_PARAMETER_STRUCT()
 
@@ -121,7 +119,7 @@ void FDMXPixelMappingRenderer::DownsampleRender(
 	const FTextureResource* InputTexture,
 	const FTextureResource* DstTexture,
 	const FTextureRenderTargetResource* DstTextureTargetResource,
-	const TArray<FDMXPixelMappingDownsamplePixelParam>& InDownsamplePixelPass,
+	const TArray<FDMXPixelMappingDownsamplePixelParamsV2>& InDownsamplePixelPass,
 	DownsampleReadCallback InCallback
 ) const
 {
@@ -133,7 +131,7 @@ void FDMXPixelMappingRenderer::DownsampleRender(
 		{
 			SCOPED_GPU_STAT(RHICmdList, DMXPixelMappingShadersStat);
 			SCOPED_DRAW_EVENTF(RHICmdList, DMXPixelMappingShadersStat, DMXPixelMappingRenderer::RenderPassName);
-
+			
 			FRHITexture* RenderTargetRef = DstTextureTargetResource->TextureRHI;
 			FRHITexture* DstTextureRef = DstTexture->TextureRHI;
 			FRHITexture* ResolveRenderTarget = DstTextureTargetResource->GetRenderTargetTexture();
@@ -159,7 +157,7 @@ void FDMXPixelMappingRenderer::DownsampleRender(
 			{
 				RHICmdList.SetViewport(0.f, 0.f, 0.f, OutputTextureSize.X, OutputTextureSize.Y, 1.f);
 
-				for (const FDMXPixelMappingDownsamplePixelParam& PixelParam : DownsamplePixelPass)
+				for (const FDMXPixelMappingDownsamplePixelParamsV2& PixelParam : DownsamplePixelPass)
 				{
 					// Create shader permutations
 					FDMXPixelMappingRendererPS::FPermutationDomain PermutationVector;
@@ -191,15 +189,13 @@ void FDMXPixelMappingRenderer::DownsampleRender(
 						1.f / OutputTextureSize.X, 1.f / OutputTextureSize.Y,
 						1.f / TextureSize.X, 1.f / TextureSize.Y);
 					SetShaderParameters(RHICmdList, VertexShader, VertexShader.GetVertexShader(), VSParameters);
-
+					
 					// Set pixel shader buffer
 					FDMXPixelMappingRendererPS::FParameters PSParameters;
 					PSParameters.InputTexture = InputTextureRHI;
 					PSParameters.InputSampler = TStaticSamplerState<SF_Trilinear, AM_Clamp, AM_Clamp, AM_Clamp>::GetRHI();
 					PSParameters.InputTextureSize = InputTextureSize;
 					PSParameters.OutputTextureSize = OutputTextureSize;
-					PSParameters.PixelFactor = FVector4f(PixelParam.PixelFactor);
-					PSParameters.InvertPixel = PixelParam.InvertPixel;
 					PSParameters.UVCellSize = FVector2f(PixelParam.UVCellSize);
 					SetShaderParameters(RHICmdList, PixelShader, PixelShader.GetPixelShader(), PSParameters);
 
diff --git a/Engine/Plugins/VirtualProduction/DMX/DMXPixelMapping/Source/DMXPixelMappingRenderer/Private/DMXPixelMappingRenderer.h b/Engine/Plugins/VirtualProduction/DMX/DMXPixelMapping/Source/DMXPixelMappingRenderer/Private/DMXPixelMappingRenderer.h
--- a/Engine/Plugins/VirtualProduction/DMX/DMXPixelMapping/Source/DMXPixelMappingRenderer/Private/DMXPixelMappingRenderer.h
+++ b/Engine/Plugins/VirtualProduction/DMX/DMXPixelMapping/Source/DMXPixelMappingRenderer/Private/DMXPixelMappingRenderer.h
@@ -27,7 +27,7 @@ public:
 		const FTextureResource* InputTexture,
 		const FTextureResource* DstTexture,
 		const FTextureRenderTargetResource* DstTextureTargetResource,
-		const TArray<FDMXPixelMappingDownsamplePixelParam>& InDownsamplePixelPass,
+		const TArray<FDMXPixelMappingDownsamplePixelParamsV2>& InDownsamplePixelPass,
 		DownsampleReadCallback InCallback
 	) const override;
 
diff --git a/Engine/Plugins/VirtualProduction/DMX/DMXPixelMapping/Source/DMXPixelMappingRenderer/Public/IDMXPixelMappingRenderer.h b/Engine/Plugins/VirtualProduction/DMX/DMXPixelMapping/Source/DMXPixelMappingRenderer/Public/IDMXPixelMappingRenderer.h
--- a/Engine/Plugins/VirtualProduction/DMX/DMXPixelMapping/Source/DMXPixelMappingRenderer/Public/IDMXPixelMappingRenderer.h
+++ b/Engine/Plugins/VirtualProduction/DMX/DMXPixelMapping/Source/DMXPixelMappingRenderer/Public/IDMXPixelMappingRenderer.h
@@ -46,6 +46,7 @@ struct FDMXPixelMappingDownsamplePixelPreviewParam
  * Downsample pixel rendering params
  * Using for pixel rendering in downsample rendering pipeline
  */
+struct UE_DEPRECATED(5.2, "Deprecated in favor of FDMXPixelMappingDownsamplePixelParamsV2. To apply color spaces, all color values are now computed at all times.") FDMXPixelMappingDownsamplePixelParam;
 struct FDMXPixelMappingDownsamplePixelParam
 {
 	/** RGBA pixel multiplication */
@@ -73,6 +74,28 @@ struct FDMXPixelMappingDownsamplePixelParam
 	bool bStaticCalculateUV;
 };
 
+struct FDMXPixelMappingDownsamplePixelParamsV2
+{
+	/** Position in screen pixels of the top left corner of the quad */
+	FIntPoint Position;
+
+	/** Position in texels of the top left corner of the quad's UV's */
+	FVector2D UV;
+
+	/** Size in texels of the quad's total UV space */
+	FVector2D UVSize;
+
+	/** Size in texels of UV.May match UVSize */
+	FVector2D UVCellSize;
+
+	/** The quality of color samples in the pixel shader(number of samples) */
+	EDMXPixelBlendingQuality CellBlendingQuality;
+
+	/** Calculates the UV point to sample purely on the UV position/size. Works best for renderers which represent a single pixel */
+	bool bStaticCalculateUV;
+};
+
+
 /**
  * The public interface of the Pixel Mapping renderer instance interface.
  */
@@ -98,7 +121,7 @@ public:
 		const FTextureResource* InputTexture,
 		const FTextureResource* DstTexture,
 		const FTextureRenderTargetResource* DstTextureTargetResource,
-		const TArray<FDMXPixelMappingDownsamplePixelParam>& InDownsamplePixelPass,
+		const TArray<FDMXPixelMappingDownsamplePixelParamsV2>& InDownsamplePixelPass,
 		DownsampleReadCallback InCallback
 	) const = 0;
 
diff --git a/Engine/Plugins/VirtualProduction/DMX/DMXPixelMapping/Source/DMXPixelMappingRuntime/DMXPixelMappingRuntime.build.cs b/Engine/Plugins/VirtualProduction/DMX/DMXPixelMapping/Source/DMXPixelMappingRuntime/DMXPixelMappingRuntime.build.cs
--- a/Engine/Plugins/VirtualProduction/DMX/DMXPixelMapping/Source/DMXPixelMappingRuntime/DMXPixelMappingRuntime.build.cs
+++ b/Engine/Plugins/VirtualProduction/DMX/DMXPixelMapping/Source/DMXPixelMappingRuntime/DMXPixelMappingRuntime.build.cs
@@ -11,6 +11,7 @@ public class DMXPixelMappingRuntime : ModuleRules
 		PublicDependencyModuleNames.AddRange(
 			new string[]
 			{
+				"ColorManagement",
 				"Core",
 				"CoreUObject",
 				"DMXRuntime",
diff --git a/Engine/Plugins/VirtualProduction/DMX/DMXPixelMapping/Source/DMXPixelMappingRuntime/Private/Blueprint/DMXPixelMappingSubsystem.cpp b/Engine/Plugins/VirtualProduction/DMX/DMXPixelMapping/Source/DMXPixelMappingRuntime/Private/Blueprint/DMXPixelMappingSubsystem.cpp
--- a/Engine/Plugins/VirtualProduction/DMX/DMXPixelMapping/Source/DMXPixelMappingRuntime/Private/Blueprint/DMXPixelMappingSubsystem.cpp
+++ b/Engine/Plugins/VirtualProduction/DMX/DMXPixelMapping/Source/DMXPixelMappingRuntime/Private/Blueprint/DMXPixelMappingSubsystem.cpp
@@ -3,21 +3,14 @@
 #include "Blueprint/DMXPixelMappingSubsystem.h"
 #include "DMXPixelMapping.h"
 #include "Components/DMXPixelMappingRendererComponent.h"
+#include "Components/DMXPixelMappingRootComponent.h"
+#include "Components/DMXPixelMappingOutputComponent.h"
 #include "Components/DMXPixelMappingOutputDMXComponent.h"
 #include "Components/DMXPixelMappingFixtureGroupComponent.h"
 #include "Components/DMXPixelMappingMatrixComponent.h"
 
 #include "Engine/Engine.h"
 
-#define DMX_PIXEL_MAPPING_GET_SUBSYSTEM_COMPONENT(ComponentType)	\
-UDMXPixelMapping##ComponentType* UDMXPixelMappingSubsystem::Get##ComponentType(UDMXPixelMapping* InDMXPixelMapping, const FName& InComponentName) \
-{ \
-if (InDMXPixelMapping != nullptr) \
-{ \
-return InDMXPixelMapping->FindComponentOfClass<UDMXPixelMapping##ComponentType>(InComponentName); \
-} \
-return nullptr; \
-}
 
 /*static*/ UDMXPixelMappingSubsystem* UDMXPixelMappingSubsystem::GetDMXPixelMappingSubsystem_Pure()
 {
@@ -34,7 +27,40 @@ UDMXPixelMapping* UDMXPixelMappingSubsystem::GetDMXPixelMapping(UDMXPixelMapping
 	return InPixelMapping;
 }
 
-DMX_PIXEL_MAPPING_GET_SUBSYSTEM_COMPONENT(RendererComponent)
-DMX_PIXEL_MAPPING_GET_SUBSYSTEM_COMPONENT(OutputDMXComponent)
-DMX_PIXEL_MAPPING_GET_SUBSYSTEM_COMPONENT(FixtureGroupComponent)
-DMX_PIXEL_MAPPING_GET_SUBSYSTEM_COMPONENT(MatrixComponent)
+UDMXPixelMappingRendererComponent* UDMXPixelMappingSubsystem::GetRendererComponent(UDMXPixelMapping* InDMXPixelMapping, const FName& InComponentName)
+{
+	UDMXPixelMappingRootComponent* RootComponent = InDMXPixelMapping ? InDMXPixelMapping->GetRootComponent() : nullptr;
+	if (RootComponent)
+	{
+		// Use the specialized method here that is faster than the generic UDMXPixelMappingBaseComponent::FindComponentOfClass
+		return RootComponent->FindRendererComponentByName(InComponentName);
+	}
+	return nullptr;
+}
+
+UDMXPixelMappingOutputDMXComponent* UDMXPixelMappingSubsystem::GetOutputDMXComponent(UDMXPixelMapping* InDMXPixelMapping, const FName& InComponentName)
+{
+	if (InDMXPixelMapping != nullptr)
+	{
+		return InDMXPixelMapping->FindComponentOfClass<UDMXPixelMappingOutputDMXComponent>(InComponentName);
+	}
+	return nullptr;
+}
+
+UDMXPixelMappingFixtureGroupComponent* UDMXPixelMappingSubsystem::GetFixtureGroupComponent(UDMXPixelMapping* InDMXPixelMapping, const FName& InComponentName)
+{
+	if (InDMXPixelMapping != nullptr)
+	{
+		return InDMXPixelMapping->FindComponentOfClass<UDMXPixelMappingFixtureGroupComponent>(InComponentName);
+	}
+	return nullptr;
+}
+
+UDMXPixelMappingMatrixComponent* UDMXPixelMappingSubsystem::GetMatrixComponent(UDMXPixelMapping* InDMXPixelMapping, const FName& InComponentName)
+{
+	if (InDMXPixelMapping != nullptr)
+	{
+		return InDMXPixelMapping->FindComponentOfClass<UDMXPixelMappingMatrixComponent>(InComponentName);
+	}
+	return nullptr;
+}
diff --git a/Engine/Plugins/VirtualProduction/DMX/DMXPixelMapping/Source/DMXPixelMappingRuntime/Private/ColorSpace/DMXPixelMappingColorSpace.cpp b/Engine/Plugins/VirtualProduction/DMX/DMXPixelMapping/Source/DMXPixelMappingRuntime/Private/ColorSpace/DMXPixelMappingColorSpace.cpp
new file mode 100644
--- /dev/null
+++ b/Engine/Plugins/VirtualProduction/DMX/DMXPixelMapping/Source/DMXPixelMappingRuntime/Private/ColorSpace/DMXPixelMappingColorSpace.cpp
@@ -0,0 +1,37 @@
+// Copyright Epic Games, Inc. All Rights Reserved.
+
+#pragma once
+
+#include "ColorSpace/DMXPixelMappingColorSpace.h"
+
+#include "ColorSpace.h"
+#include "CoreMinimal.h"
+#include "UObject/Object.h"
+
+
+void UDMXPixelMappingColorSpace::ResetToBlack()
+{
+	for (TTuple<FDMXAttributeName, float>& AttributeNameToValuePair : AttributeNameToValueMap)
+	{
+		AttributeNameToValuePair.Value = 0.f;
+	}
+}
+
+#if WITH_EDITOR
+void UDMXPixelMappingColorSpace::PostEditChangeProperty(FPropertyChangedEvent& PropertyChangedEvent)
+{
+	Super::PostEditChangeProperty(PropertyChangedEvent);
+
+	OnPostEditChangeProperty.Broadcast(PropertyChangedEvent);
+}
+#endif // WITH_EDITOR
+
+void UDMXPixelMappingColorSpace::SetAttributeValue(const FDMXAttributeName& AttributeName, float Value)
+{
+	AttributeNameToValueMap.FindOrAdd(AttributeName) = Value;
+}
+
+void UDMXPixelMappingColorSpace::ClearCachedAttributeValues()
+{
+	AttributeNameToValueMap.Reset();
+}
diff --git a/Engine/Plugins/VirtualProduction/DMX/DMXPixelMapping/Source/DMXPixelMappingRuntime/Private/ColorSpace/DMXPixelMappingColorSpace_RGBCMY.cpp b/Engine/Plugins/VirtualProduction/DMX/DMXPixelMapping/Source/DMXPixelMappingRuntime/Private/ColorSpace/DMXPixelMappingColorSpace_RGBCMY.cpp
new file mode 100644
--- /dev/null
+++ b/Engine/Plugins/VirtualProduction/DMX/DMXPixelMapping/Source/DMXPixelMappingRuntime/Private/ColorSpace/DMXPixelMappingColorSpace_RGBCMY.cpp
@@ -0,0 +1,123 @@
+// Copyright Epic Games, Inc. All Rights Reserved.
+
+#include "ColorSpace/DMXPixelMappingColorSpace_RGBCMY.h"
+
+
+UDMXPixelMappingColorSpace_RGBCMY::UDMXPixelMappingColorSpace_RGBCMY()
+	: RedAttribute("Red")
+	, GreenAttribute("Green")
+	, BlueAttribute("Blue")
+	, LuminanceAttribute("Dimmer")
+{}
+
+void UDMXPixelMappingColorSpace_RGBCMY::SetRGBA(const FLinearColor& InColor)
+{
+	if (bUseWorkingColorSpaceForInput && !InputColorSpace.Equals(UE::Color::FColorSpace::GetWorking()))
+	{
+		// Update in case the working color space changed
+		UpdateColorSpaceAndTransform();
+	}
+
+	const FLinearColor CalibratedColor = ColorSpaceTransform->Apply(InColor);
+
+	if (RedAttribute.IsValid())
+	{
+		const float Value = bSendCyan ? FMath::Abs(CalibratedColor.R - 1.f) : CalibratedColor.R;
+		SetAttributeValue(RedAttribute, Value);
+	}
+
+	if (GreenAttribute.IsValid())
+	{
+		const float Value = bSendMagenta ? FMath::Abs(CalibratedColor.G - 1.f) : CalibratedColor.G;
+		SetAttributeValue(GreenAttribute, Value);
+	}
+
+	if (BlueAttribute.IsValid())
+	{
+		const float Value = bSendYellow ? FMath::Abs(CalibratedColor.B - 1.f) : CalibratedColor.B;
+		SetAttributeValue(BlueAttribute, Value);
+	}
+
+	if (LuminanceAttribute.IsValid())
+	{
+		if (LuminanceType == EDMXPixelMappingLuminanceType_RGBCMY::FromColor)
+		{
+			const float LuminanceFromColor = OutputColorSpace.GetLuminance(CalibratedColor);
+			SetAttributeValue(LuminanceAttribute, LuminanceFromColor);
+		}
+		else if (LuminanceType == EDMXPixelMappingLuminanceType_RGBCMY::Constant)
+		{
+			SetAttributeValue(LuminanceAttribute, Luminance);
+		}
+		else if (LuminanceType == EDMXPixelMappingLuminanceType_RGBCMY::FromAlpha)
+		{
+			SetAttributeValue(LuminanceAttribute, InColor.A);
+		}
+	}
+}
+
+void UDMXPixelMappingColorSpace_RGBCMY::PostLoad()
+{
+	Super::PostLoad();
+
+	UpdateColorSpaceAndTransform();
+}
+
+#if WITH_EDITOR
+void UDMXPixelMappingColorSpace_RGBCMY::PostEditChangeChainProperty(FPropertyChangedChainEvent& PropertyChangedEvent)
+{
+	Super::PostEditChangeChainProperty(PropertyChangedEvent);
+
+	const FName PropertyName = PropertyChangedEvent.GetPropertyName();
+
+	if (PropertyName == GET_MEMBER_NAME_CHECKED(FDMXAttributeName, Name))
+	{
+		ClearCachedAttributeValues();
+	}
+	else if (PropertyName == GET_MEMBER_NAME_CHECKED(UDMXPixelMappingColorSpace_RGBCMY, OutputColorSpace) ||
+		PropertyName == GET_MEMBER_NAME_CHECKED(UDMXPixelMappingColorSpace, bUseWorkingColorSpaceForInput))
+	{
+		UpdateColorSpaceAndTransform();
+	}
+}
+#endif // WITH_EDITOR
+
+void UDMXPixelMappingColorSpace_RGBCMY::UpdateColorSpaceAndTransform()
+{
+	using namespace UE::Color;
+
+	InputColorSpace = bUseWorkingColorSpaceForInput ? FColorSpace::GetWorking() : FColorSpace(EColorSpace::sRGB);
+
+	const EColorSpace SelectedOutputColorSpace = ConvertToOutputColorSpaceEnum(PixelMappingOutputColorSpace);
+	OutputColorSpace = FColorSpace(SelectedOutputColorSpace);
+
+	ColorSpaceTransform = MakeUnique<FColorSpaceTransform>(InputColorSpace, OutputColorSpace);
+}
+
+UE::Color::EColorSpace UDMXPixelMappingColorSpace_RGBCMY::ConvertToOutputColorSpaceEnum(EDMXPixelMappingOutputColorSpace_RGBCMY InPixelMappingOutputColorSpaceEnum) const
+{
+	using namespace UE::Color;
+	if (InPixelMappingOutputColorSpaceEnum == EDMXPixelMappingOutputColorSpace_RGBCMY::sRGB)
+	{
+		return EColorSpace::sRGB;
+	}
+	else if (InPixelMappingOutputColorSpaceEnum == EDMXPixelMappingOutputColorSpace_RGBCMY::Plasa)
+	{
+		return EColorSpace::PLASA_E1_54;
+	}
+	else if (InPixelMappingOutputColorSpaceEnum == EDMXPixelMappingOutputColorSpace_RGBCMY::Rec2020)
+	{
+		return EColorSpace::Rec2020;
+	}
+	else if (InPixelMappingOutputColorSpaceEnum == EDMXPixelMappingOutputColorSpace_RGBCMY::P3DCI)
+	{
+		return EColorSpace::P3DCI;
+	}
+	else if (InPixelMappingOutputColorSpaceEnum == EDMXPixelMappingOutputColorSpace_RGBCMY::P3D65)
+	{
+		return EColorSpace::P3D65;
+	}
+
+	return EColorSpace::None;
+	checkf(0, TEXT("Unhandled PixelMapping Color Space enum used in UDMXPixelMappingColorSpace::ConvertToColorSpaceEnum"));
+}
diff --git a/Engine/Plugins/VirtualProduction/DMX/DMXPixelMapping/Source/DMXPixelMappingRuntime/Private/ColorSpace/DMXPixelMappingColorSpace_RGBCMY.h b/Engine/Plugins/VirtualProduction/DMX/DMXPixelMapping/Source/DMXPixelMappingRuntime/Private/ColorSpace/DMXPixelMappingColorSpace_RGBCMY.h
new file mode 100644
--- /dev/null
+++ b/Engine/Plugins/VirtualProduction/DMX/DMXPixelMapping/Source/DMXPixelMappingRuntime/Private/ColorSpace/DMXPixelMappingColorSpace_RGBCMY.h
@@ -0,0 +1,118 @@
+// Copyright Epic Games, Inc. All Rights Reserved.
+
+#pragma once
+
+#include "DMXAttribute.h"
+#include "ColorSpace/DMXPixelMappingColorSpace.h"
+
+#include "ColorSpace.h"
+#include "Templates/UniquePtr.h"
+
+#include "DMXPixelMappingColorSpace_RGBCMY.generated.h"
+
+
+UENUM(BlueprintType)
+enum class EDMXPixelMappingOutputColorSpace_RGBCMY : uint8
+{
+	sRGB UMETA(DisplayName = "sRGB / Rec709"),
+	Rec2020,
+	P3DCI,
+	P3D65,
+	Plasa UMETA(DisplayName = "PLASA RGB - ANSI E1.54")
+};
+
+UENUM(BlueprintType)
+enum class EDMXPixelMappingLuminanceType_RGBCMY : uint8
+{
+	FromColor,
+	Constant,
+	FromAlpha,
+	None
+};
+
+UCLASS(meta = (DisplayName = "RGB / CMY"))
+class DMXPIXELMAPPINGRUNTIME_API UDMXPixelMappingColorSpace_RGBCMY
+	: public UDMXPixelMappingColorSpace
+{
+	GENERATED_BODY()
+
+public:
+	/** Constructor */
+	UDMXPixelMappingColorSpace_RGBCMY();
+
+	//~ Begin DMXPixelMappingColorSpace interface
+	virtual void SetRGBA(const FLinearColor& InColor) override;
+	//~ End DMXPixelMappingColorSpace interface
+
+	/** The color space to use */
+	UPROPERTY(EditAnywhere, Category = "Color Space", Meta = (DisplayPriority = 1, DisplayName = "Output Color Space"))
+	EDMXPixelMappingOutputColorSpace_RGBCMY PixelMappingOutputColorSpace = EDMXPixelMappingOutputColorSpace_RGBCMY::sRGB;
+
+	/** If set, converts Red to Cyan */
+	UPROPERTY(EditAnywhere, Category = "RGB")
+	bool bSendCyan = false;
+
+	/** If set, converts Green to Magenta */
+	UPROPERTY(EditAnywhere, Category = "RGB")
+	bool bSendMagenta = false;
+
+	/** If set, converts Blue to Yellow */
+	UPROPERTY(EditAnywhere, Category = "RGB")
+	bool bSendYellow = false;
+
+	/** Attribute sent for Red */
+	UPROPERTY(EditAnywhere, Category = "RGB")
+	FDMXAttributeName RedAttribute;
+
+	/** Attribute sent for Green */
+	UPROPERTY(EditAnywhere, Category = "RGB")
+	FDMXAttributeName GreenAttribute;
+
+	/** Attribute sent for Blue */
+	UPROPERTY(EditAnywhere, Category = "RGB")
+	FDMXAttributeName BlueAttribute;
+
+	/** Adds a Dimmer Attribute */
+	UPROPERTY(EditAnywhere, Category = "Luminance")
+	EDMXPixelMappingLuminanceType_RGBCMY LuminanceType = EDMXPixelMappingLuminanceType_RGBCMY::FromColor;
+
+	/** Attribute sent for the generated Luminance */
+	UPROPERTY(EditAnywhere, Category = "Luminance", Meta = (EditConditionHides, EditCondition = "LuminanceType != EDMXPixelMappingLuminanceType_RGBCMY::None"))
+	FDMXAttributeName LuminanceAttribute;
+
+	/** Luminance used when Luminance Type is set to 'Set Value' */
+	UPROPERTY(EditAnywhere, Category = "Luminance", Meta = (ClampMin = 0.0, ClampMax = 1.0, UIMin = 0.0, UIMax = 1.0, EditConditionHides, EditCondition = "LuminanceType == EDMXPixelMappingLuminanceType_RGBCMY::Constant"))
+	float Luminance = 1.f;
+
+	/** Min Luminance used when Luminance Type is set to 'From White' or 'From Alpha' */
+	UPROPERTY(EditAnywhere, Category = "Luminance", Meta = (ClampMin = 0.0, ClampMax = 1.0, UIMin = 0.0, UIMax = 1.0, EditConditionHides, EditCondition = "LuminanceType == EDMXPixelMappingLuminanceType_RGBCMY::FromColor || LuminanceType == EDMXPixelMappingLuminanceType_RGBCMY::FromAlpha"))
+	float MinLuminance = 0.f;
+
+	/** Max Luminance used when Luminance Type is set to 'From White' or 'From Alpha' */
+	UPROPERTY(EditAnywhere, Category = "Luminance", Meta = (ClampMin = 0.0, ClampMax = 1.0, UIMin = 0.0, UIMax = 1.0, EditConditionHides, EditCondition = "LuminanceType == EDMXPixelMappingLuminanceType_RGBCMY::FromColor || LuminanceType == EDMXPixelMappingLuminanceType_RGBCMY::FromAlpha"))
+	float MaxLuminance = 1.f;
+
+protected:
+	//~ Begin UObject Interface
+	virtual void PostLoad() override;
+#if WITH_EDITOR
+	virtual void PostEditChangeChainProperty(FPropertyChangedChainEvent& PropertyChangedEvent) override;
+#endif // WITH_EDITOR
+	//~ End UObject Interface
+
+private:
+	/** Updates the ColorSpace and ColorSpaceTransform members */
+	void UpdateColorSpaceAndTransform();
+
+	/** Gets the EColorSpace enum given the current Output Color Space Enum */
+	UE::Color::EColorSpace ConvertToOutputColorSpaceEnum(EDMXPixelMappingOutputColorSpace_RGBCMY InPixelMappingOutputColorSpaceEnum) const;
+
+	/** The input color space instance */
+	UE::Color::FColorSpace InputColorSpace;
+
+	/** The output color space instance */
+	UE::Color::FColorSpace OutputColorSpace;
+
+	/** The color space transform instance */
+	TUniquePtr<UE::Color::FColorSpaceTransform> ColorSpaceTransform;
+};
diff --git a/Engine/Plugins/VirtualProduction/DMX/DMXPixelMapping/Source/DMXPixelMappingRuntime/Private/ColorSpace/DMXPixelMappingColorSpace_XYZ.cpp b/Engine/Plugins/VirtualProduction/DMX/DMXPixelMapping/Source/DMXPixelMappingRuntime/Private/ColorSpace/DMXPixelMappingColorSpace_XYZ.cpp
new file mode 100644
--- /dev/null
+++ b/Engine/Plugins/VirtualProduction/DMX/DMXPixelMapping/Source/DMXPixelMappingRuntime/Private/ColorSpace/DMXPixelMappingColorSpace_XYZ.cpp
@@ -0,0 +1,52 @@
+// Copyright Epic Games, Inc. All Rights Reserved.
+
+#include "ColorSpace/DMXPixelMappingColorSpace_XYZ.h"
+
+
+UDMXPixelMappingColorSpace_XYZ::UDMXPixelMappingColorSpace_XYZ()
+	: XAttribute("X")
+	, ZAttribute("Z")
+	, LuminanceAttribute("Dimmer")
+{}
+
+void UDMXPixelMappingColorSpace_XYZ::SetRGBA(const FLinearColor& InColor)
+{
+	const UE::Color::FColorSpace& InputColorSpace = bUseWorkingColorSpaceForInput ?
+		UE::Color::FColorSpace::GetWorking() :
+		SRGBColorSpace;
+
+	// Convert RGB to CIEXYZ
+	const FMatrix44d Matrix = InputColorSpace.GetRgbToXYZ();
+	const FVector4 XYZW = Matrix.TransformVector(FVector(InColor));
+
+	if (XAttribute.IsValid())
+	{
+		SetAttributeValue(XAttribute, XYZW.X);
+	}
+
+	if (LuminanceAttribute.IsValid())
+	{
+		SetAttributeValue(LuminanceAttribute, XYZW.Y);
+	}
+
+	if (ZAttribute.IsValid())
+	{
+		SetAttributeValue(ZAttribute, XYZW.Z);
+	}
+}
+
+#if WITH_EDITOR
+void UDMXPixelMappingColorSpace_XYZ::PostEditChangeProperty(FPropertyChangedEvent& PropertyChangedEvent)
+{
+	Super::PostEditChangeProperty(PropertyChangedEvent);
+
+	const FName PropertyName = PropertyChangedEvent.GetPropertyName();
+
+	if (PropertyName == GET_MEMBER_NAME_CHECKED(UDMXPixelMappingColorSpace_XYZ, XAttribute) ||
+		PropertyName == GET_MEMBER_NAME_CHECKED(UDMXPixelMappingColorSpace_XYZ, LuminanceAttribute) ||
+		PropertyName == GET_MEMBER_NAME_CHECKED(UDMXPixelMappingColorSpace_XYZ, ZAttribute))
+	{
+		ClearCachedAttributeValues();
+	}
+}
+#endif // WITH_EDITOR
diff --git a/Engine/Plugins/VirtualProduction/DMX/DMXPixelMapping/Source/DMXPixelMappingRuntime/Private/ColorSpace/DMXPixelMappingColorSpace_XYZ.h b/Engine/Plugins/VirtualProduction/DMX/DMXPixelMapping/Source/DMXPixelMappingRuntime/Private/ColorSpace/DMXPixelMappingColorSpace_XYZ.h
new file mode 100644
--- /dev/null
+++ b/Engine/Plugins/VirtualProduction/DMX/DMXPixelMapping/Source/DMXPixelMappingRuntime/Private/ColorSpace/DMXPixelMappingColorSpace_XYZ.h
@@ -0,0 +1,61 @@
+// Copyright Epic Games, Inc. All Rights Reserved.
+
+#pragma once
+
+#include "DMXAttribute.h"
+#include "ColorSpace/DMXPixelMappingColorSpace.h"
+
+#include "ColorSpace.h"
+
+#include "DMXPixelMappingColorSpace_XYZ.generated.h"
+
+
+
+UCLASS(meta = (DisplayName = "CIE 1931 XYZ"))
+class DMXPIXELMAPPINGRUNTIME_API UDMXPixelMappingColorSpace_XYZ
+	: public UDMXPixelMappingColorSpace
+{
+	GENERATED_BODY()
+
+public:
+	/** Constructor */
+	UDMXPixelMappingColorSpace_XYZ();
+
+	//~ Begin DMXPixelMappingColorSpace interface
+	virtual void SetRGBA(const FLinearColor& InColor) override;
+	//~ End DMXPixelMappingColorSpace interface
+
+	/** Attribute sent for X */
+	UPROPERTY(EditAnywhere, Category = "XYZ", Meta = (DisplayName = "X Attribute"))
+	FDMXAttributeName XAttribute;
+
+	/** Attribute sent for Z */
+	UPROPERTY(EditAnywhere, Category = "XYZ", Meta = (DisplayName = "Z Attribute"))
+	FDMXAttributeName ZAttribute;
+
+	/** Attribute sent for Y */
+	UPROPERTY(EditAnywhere, Category = "XYZ")
+	FDMXAttributeName LuminanceAttribute;
+
+	/** Scale Factor for Luminance, e.g. when 0.5, a Luminance of 1 will be output as 0.5. */
+	UPROPERTY(EditAnywhere, Category = "Luminance", Meta = (ClampMin = 0.0, ClampMax = 1.0, UIMin = 0.0, UIMax = 1.0, EditConditionHides, EditCondition = "LuminanceType == EDMXPixelMappingLuminanceType_RGBCMY::Constant"))
+	float LuminanceScale = 1.f;
+
+	/** Min Luminance */
+	UPROPERTY(EditAnywhere, Category = "Luminance", Meta = (ClampMin = 0.0, ClampMax = 1.0, UIMin = 0.0, UIMax = 1.0, EditConditionHides, EditCondition = "LuminanceType == EDMXPixelMappingLuminanceType_RGBCMY::FromWhite || LuminanceType == EDMXPixelMappingLuminanceType_RGBCMY::FromAlpha"))
+	float MinLuminance = 0.f;
+
+	/** Max Luminance */
+	UPROPERTY(EditAnywhere, Category = "Luminance", Meta = (ClampMin = 0.0, ClampMax = 1.0, UIMin = 0.0, UIMax = 1.0, EditConditionHides, EditCondition = "LuminanceType == EDMXPixelMappingLuminanceType_RGBCMY::FromWhite || LuminanceType == EDMXPixelMappingLuminanceType_RGBCMY::FromAlpha"))
+	float MaxLuminance = 1.f;
+
+protected:
+	//~ Begin UObject Interface
+#if WITH_EDITOR
+	virtual void PostEditChangeProperty(FPropertyChangedEvent& PropertyChangedEvent) override;
+#endif // WITH_EDITOR
+
+private:
+	/** Cached sRGB color space, to avoid instantiating on each conversion */
+	UE::Color::FColorSpace SRGBColorSpace;
+};
diff --git a/Engine/Plugins/VirtualProduction/DMX/DMXPixelMapping/Source/DMXPixelMappingRuntime/Private/ColorSpace/DMXPixelMappingColorSpace_xyY.cpp b/Engine/Plugins/VirtualProduction/DMX/DMXPixelMapping/Source/DMXPixelMappingRuntime/Private/ColorSpace/DMXPixelMappingColorSpace_xyY.cpp
new file mode 100644
--- /dev/null
+++ b/Engine/Plugins/VirtualProduction/DMX/DMXPixelMapping/Source/DMXPixelMappingRuntime/Private/ColorSpace/DMXPixelMappingColorSpace_xyY.cpp
@@ -0,0 +1,60 @@
+// Copyright Epic Games, Inc. All Rights Reserved.
+
+#include "ColorSpace/DMXPixelMappingColorSpace_xyY.h"
+
+
+UDMXPixelMappingColorSpace_xyY::UDMXPixelMappingColorSpace_xyY()
+	: XAttribute("CIE_X")
+	, YAttribute("CIE_Y")
+	, LuminanceAttribute("Dimmer")
+	, SRGBColorSpace(UE::Color::EColorSpace::sRGB)
+{}
+
+void UDMXPixelMappingColorSpace_xyY::SetRGBA(const FLinearColor& InColor)
+{
+	const UE::Color::FColorSpace& InputColorSpace = bUseWorkingColorSpaceForInput ?
+		UE::Color::FColorSpace::GetWorking() :
+		SRGBColorSpace;
+
+	// Convert RGB to CIE XYZ
+	const FMatrix44d Matrix = InputColorSpace.GetRgbToXYZ();
+	const FVector4 XYZW = Matrix.TransformVector(FVector(InColor));
+
+	const double SumXYZ = XYZW.X + XYZW.Y + XYZW.Z;
+
+	// Calculate chromaticity coordinates x and y, use white point if SumXYZ is zero
+	const double x = SumXYZ > 0.f ? XYZW.X / SumXYZ : 0.3127;
+	const double y = SumXYZ > 0.f ? XYZW.Y / SumXYZ : 0.3290;
+
+	if (XAttribute.IsValid())
+	{
+		SetAttributeValue(XAttribute, x);
+	}
+
+	if (YAttribute.IsValid())
+	{
+		SetAttributeValue(YAttribute, y);
+	}
+
+	if (LuminanceAttribute.IsValid())
+	{
+		const float Luminance = FMath::Clamp(XYZW.Y * LuminanceScale, MinLuminance, MaxLuminance);
+		SetAttributeValue(LuminanceAttribute, Luminance);
+	}
+}
+
+#if WITH_EDITOR
+void UDMXPixelMappingColorSpace_xyY::PostEditChangeProperty(FPropertyChangedEvent& PropertyChangedEvent)
+{
+	Super::PostEditChangeProperty(PropertyChangedEvent);
+
+	const FName PropertyName = PropertyChangedEvent.GetPropertyName();
+
+	if (PropertyName == GET_MEMBER_NAME_CHECKED(UDMXPixelMappingColorSpace_xyY, XAttribute) ||
+		PropertyName == GET_MEMBER_NAME_CHECKED(UDMXPixelMappingColorSpace_xyY, YAttribute) ||
+		PropertyName == GET_MEMBER_NAME_CHECKED(UDMXPixelMappingColorSpace_xyY, LuminanceAttribute))
+	{
+		ClearCachedAttributeValues();
+	}
+}
+#endif // WITH_EDITOR
diff --git a/Engine/Plugins/VirtualProduction/DMX/DMXPixelMapping/Source/DMXPixelMappingRuntime/Private/ColorSpace/DMXPixelMappingColorSpace_xyY.h b/Engine/Plugins/VirtualProduction/DMX/DMXPixelMapping/Source/DMXPixelMappingRuntime/Private/ColorSpace/DMXPixelMappingColorSpace_xyY.h
new file mode 100644
--- /dev/null
+++ b/Engine/Plugins/VirtualProduction/DMX/DMXPixelMapping/Source/DMXPixelMappingRuntime/Private/ColorSpace/DMXPixelMappingColorSpace_xyY.h
@@ -0,0 +1,60 @@
+// Copyright Epic Games, Inc. All Rights Reserved.
+
+#pragma once
+
+#include "DMXAttribute.h"
+#include "ColorSpace/DMXPixelMappingColorSpace.h"
+
+#include "ColorSpace.h"
+
+#include "DMXPixelMappingColorSpace_xyY.generated.h"
+
+
+UCLASS(meta = (DisplayName = "CIE 1931 xyY"))
+class DMXPIXELMAPPINGRUNTIME_API UDMXPixelMappingColorSpace_xyY
+	: public UDMXPixelMappingColorSpace
+{
+	GENERATED_BODY()
+
+public:
+	/** Constructor */
+	UDMXPixelMappingColorSpace_xyY();
+
+	//~ Begin DMXPixelMappingColorSpace interface
+	virtual void SetRGBA(const FLinearColor& InColor) override;
+	//~ End DMXPixelMappingColorSpace interface
+
+	/** Attribute sent for x */
+	UPROPERTY(EditAnywhere, Category = "XY", Meta = (DisplayName = "x Attribute"))
+	FDMXAttributeName XAttribute;
+
+	/** Attribute sent for y */
+	UPROPERTY(EditAnywhere, Category = "XY", Meta = (DisplayName = "y Attribute"))
+	FDMXAttributeName YAttribute;
+
+	/** Attribute sent for Y */
+	UPROPERTY(EditAnywhere, Category = "Luminance")
+	FDMXAttributeName LuminanceAttribute;
+
+	/** Scale Factor for Luminance, e.g. when 0.5, a Luminance of 1 will be output as 0.5. */
+	UPROPERTY(EditAnywhere, Category = "Luminance", Meta = (ClampMin = 0.0, ClampMax = 1.0, SliderMin = 0.0, SliderMax = 1.0, EditConditionHides, EditCondition = "LuminanceType == EDMXPixelMappingLuminanceType_RGBCMY::Constant"))
+	float LuminanceScale = 1.f;
+
+	/** Min Luminance */
+	UPROPERTY(EditAnywhere, Category = "Luminance", Meta = (ClampMin = 0.0, ClampMax = 1.0, SliderMin = 0.0, SliderMax = 1.0, EditConditionHides, EditCondition = "LuminanceType == EDMXPixelMappingLuminanceType_RGBCMY::FromWhite || LuminanceType == EDMXPixelMappingLuminanceType_RGBCMY::FromAlpha"))
+	float MinLuminance = 0.f;
+
+	/** Max Luminance */
+	UPROPERTY(EditAnywhere, Category = "Luminance", Meta = (ClampMin = 0.0, ClampMax = 1.0, SliderMin = 0.0, SliderMax = 1.0, EditConditionHides, EditCondition = "LuminanceType == EDMXPixelMappingLuminanceType_RGBCMY::FromWhite || LuminanceType == EDMXPixelMappingLuminanceType_RGBCMY::FromAlpha"))
+	float MaxLuminance = 1.f;
+
+protected:
+	//~ Begin UObject Interface
+#if WITH_EDITOR
+	virtual void PostEditChangeProperty(FPropertyChangedEvent& PropertyChangedEvent) override;
+#endif // WITH_EDITOR
+
+private:
+	/** Cached sRGB color space, to avoid instantiating on each conversion */
+	UE::Color::FColorSpace SRGBColorSpace;
+};
diff --git a/Engine/Plugins/VirtualProduction/DMX/DMXPixelMapping/Source/DMXPixelMappingRuntime/Private/Components/DMXPixelMappingBaseComponent.cpp b/Engine/Plugins/VirtualProduction/DMX/DMXPixelMapping/Source/DMXPixelMappingRuntime/Private/Components/DMXPixelMappingBaseComponent.cpp
--- a/Engine/Plugins/VirtualProduction/DMX/DMXPixelMapping/Source/DMXPixelMappingRuntime/Private/Components/DMXPixelMappingBaseComponent.cpp
+++ b/Engine/Plugins/VirtualProduction/DMX/DMXPixelMapping/Source/DMXPixelMappingRuntime/Private/Components/DMXPixelMappingBaseComponent.cpp
@@ -3,12 +3,12 @@
 #include "Components/DMXPixelMappingBaseComponent.h"
 
 #include "DMXPixelMapping.h"
-#include "DMXPixelMappingRuntimeObjectVersion.h"
+#include "DMXPixelMappingRuntimeCommon.h"
+#include "DMXPixelMappingMainStreamObjectVersion.h"
 #include "Components/DMXPixelMappingRendererComponent.h"
 #include "Components/DMXPixelMappingRootComponent.h"
 
 
-
 UDMXPixelMappingBaseComponent::FDMXPixelMappingOnComponentAdded UDMXPixelMappingBaseComponent::OnComponentAdded;
 UDMXPixelMappingBaseComponent::FDMXPixelMappingOnComponentRemoved UDMXPixelMappingBaseComponent::OnComponentRemoved;
 UDMXPixelMappingBaseComponent::FDMXPixelMappingOnComponentRenamed UDMXPixelMappingBaseComponent::OnComponentRenamed;
@@ -35,10 +35,10 @@ void UDMXPixelMappingBaseComponent::Serialize(FArchive& Ar)
 {
 	Super::Serialize(Ar);
 
-	Ar.UsingCustomVersion(FDMXPixelMappingRuntimeObjectVersion::GUID);
+	Ar.UsingCustomVersion(FDMXPixelMappingMainStreamObjectVersion::GUID);
 	if (Ar.IsLoading())
 	{
-		if (Ar.CustomVer(FDMXPixelMappingRuntimeObjectVersion::GUID) < FDMXPixelMappingRuntimeObjectVersion::UseWeakPtrForPixelMappingComponentParent)
+		if (Ar.CustomVer(FDMXPixelMappingMainStreamObjectVersion::GUID) < FDMXPixelMappingMainStreamObjectVersion::UseWeakPtrForPixelMappingComponentParent)
 		{
 			// Upgrade from strong object references for parent to weak ones
 			WeakParent = Parent_DEPRECATED;
@@ -52,8 +52,7 @@ void UDMXPixelMappingBaseComponent::PostRename(UObject* OldOuter, const FName Ol
 {
 	Super::PostRename(OldOuter, OldName);
 
-	// Broadcast the change
-	GetOnComponentRenamed().Broadcast(GetPixelMapping(), this, OldOuter, OldName);
+	OnComponentRenamed.Broadcast(GetPixelMapping(), this, OldOuter, OldName);
 }
 
 #if WITH_EDITOR
@@ -205,7 +204,7 @@ void UDMXPixelMappingBaseComponent::AddChild(UDMXPixelMappingBaseComponent* InCo
 		InComponent->NotifyAddedToParent();
 
 		// Broadcast the change
-		GetOnComponentAdded().Broadcast(GetPixelMapping(), InComponent);
+		OnComponentAdded.Broadcast(GetPixelMapping(), InComponent);
 	}
 }
 
@@ -220,8 +219,7 @@ void UDMXPixelMappingBaseComponent::RemoveChild(UDMXPixelMappingBaseComponent* C
 
 		ChildComponent->SetFlags(RF_Transactional);
 
-		// Broadcast the change
-		GetOnComponentRemoved().Broadcast(GetPixelMapping(), ChildComponent);
+		OnComponentRemoved.Broadcast(GetPixelMapping(), ChildComponent);
 	}
 
 	ensure(ChildComponent && !Children.Contains(ChildComponent) && ChildComponent->Children.Num() == 0 && !ChildComponent->Parent_DEPRECATED);
diff --git a/Engine/Plugins/VirtualProduction/DMX/DMXPixelMapping/Source/DMXPixelMappingRuntime/Private/Components/DMXPixelMappingFixtureGroupItemComponent.cpp b/Engine/Plugins/VirtualProduction/DMX/DMXPixelMapping/Source/DMXPixelMappingRuntime/Private/Components/DMXPixelMappingFixtureGroupItemComponent.cpp
--- a/Engine/Plugins/VirtualProduction/DMX/DMXPixelMapping/Source/DMXPixelMappingRuntime/Private/Components/DMXPixelMappingFixtureGroupItemComponent.cpp
+++ b/Engine/Plugins/VirtualProduction/DMX/DMXPixelMapping/Source/DMXPixelMappingRuntime/Private/Components/DMXPixelMappingFixtureGroupItemComponent.cpp
@@ -2,21 +2,25 @@
 
 #include "Components/DMXPixelMappingFixtureGroupItemComponent.h"
 
+#include "DMXConversions.h"
+#include "DMXPixelMappingMainStreamObjectVersion.h"
 #include "DMXPixelMappingRuntimeUtils.h"
 #include "DMXPixelMappingTypes.h"
+#include "IDMXPixelMappingRenderer.h"
+#include "ColorSpace/DMXPixelMappingColorSpace_RGBCMY.h"
 #include "Components/DMXPixelMappingFixtureGroupComponent.h"
 #include "Components/DMXPixelMappingRendererComponent.h"
-#include "IDMXPixelMappingRenderer.h"
 #include "Library/DMXEntityFixturePatch.h"
+#include "Library/DMXEntityFixtureType.h"
 #include "Library/DMXLibrary.h"
 #include "IO/DMXOutputPort.h"
 #include "Modulators/DMXModulator.h"
-#include "TextureResource.h"
 
 #if WITH_EDITOR
 #include "DMXPixelMappingComponentWidget.h"
 #endif // WITH_EDITOR
 
+#include "TextureResource.h"
 #include "Engine/Texture.h"
 
 
@@ -27,21 +31,47 @@ DECLARE_CYCLE_STAT(TEXT("Send Fixture Group Item"), STAT_DMXPixelMaping_FixtureG
 UDMXPixelMappingFixtureGroupItemComponent::UDMXPixelMappingFixtureGroupItemComponent()
 	: DownsamplePixelIndex(0)
 {
-	SetSize(FVector2D(32.f, 32.f));
-
-	ColorMode = EDMXColorMode::CM_RGB;
-	AttributeRExpose = AttributeGExpose = AttributeBExpose = true;
-	bMonochromeExpose = true;
+	ColorSpace = CreateDefaultSubobject<UDMXPixelMappingColorSpace_RGBCMY>("ColorSpace");
 
-	AttributeR.SetFromName("Red");
-	AttributeG.SetFromName("Green");
-	AttributeB.SetFromName("Blue");
+	SetSize(FVector2D(32.f, 32.f));
 
 #if WITH_EDITOR
 	ZOrder = 2;
 #endif // WITH_EDITOR
 }
 
+void UDMXPixelMappingFixtureGroupItemComponent::Serialize(FArchive& Ar)
+{
+	Super::Serialize(Ar);
+
+#if WITH_EDITOR
+	Ar.UsingCustomVersion(FDMXPixelMappingMainStreamObjectVersion::GUID);
+	if (Ar.IsLoading())
+	{
+		// Upgrade to the DMXPixelMappingColorSpace default subobject
+		if (Ar.CustomVer(FDMXPixelMappingMainStreamObjectVersion::GUID) < FDMXPixelMappingMainStreamObjectVersion::UseDMXPixelMappingColorSpace)
+		{
+			UDMXPixelMappingColorSpace_RGBCMY* ColorSpace_RGBCMY = Cast<UDMXPixelMappingColorSpace_RGBCMY>(ColorSpace);
+			checkf(ColorSpace_RGBCMY, TEXT("Missing default Subobject ColorSpace"));
+
+			if (bMonochromeExpose_DEPRECATED)
+			{
+				ColorSpace_RGBCMY->LuminanceAttribute = MonochromeIntensity_DEPRECATED;
+			}
+			else
+			{
+				ColorSpace_RGBCMY->bSendCyan = AttributeRInvert_DEPRECATED;
+				ColorSpace_RGBCMY->bSendMagenta = AttributeGInvert_DEPRECATED;
+				ColorSpace_RGBCMY->bSendYellow = AttributeBInvert_DEPRECATED;
+				ColorSpace_RGBCMY->RedAttribute = AttributeR_DEPRECATED;
+				ColorSpace_RGBCMY->GreenAttribute = AttributeG_DEPRECATED;
+				ColorSpace_RGBCMY->BlueAttribute = AttributeB_DEPRECATED;
+			}
+		}
+	}
+#endif
+}
+
 void UDMXPixelMappingFixtureGroupItemComponent::PostLoad()
 {
 	Super::PostLoad();
@@ -64,6 +94,29 @@ void UDMXPixelMappingFixtureGroupItemComponent::PostLoad()
 			}
 		}
 	}
+
+	// Set the transient Color Space Class property to the initial value
+	if (ensureMsgf(ColorSpace, TEXT("Unexpected invalid Color Space in Pixel Mapping Fixture Group Item %s."), *GetName()))
+	{
+		ColorSpaceClass = ColorSpace->GetClass();
+
+#if WITH_EDITOR
+		ColorSpace->GetOnPostEditChangedProperty().AddUObject(this, &UDMXPixelMappingFixtureGroupItemComponent::OnColorSpacePostEditChangeProperties);
+#endif // WITH_EDITOR
+	}
+}
+
+void UDMXPixelMappingFixtureGroupItemComponent::BeginDestroy() 
+{
+	Super::BeginDestroy();
+
+#if WITH_EDITOR
+	// Set the transient Color Space Class property to the initial value
+	if (ColorSpace)
+	{
+		ColorSpace->GetOnPostEditChangedProperty().RemoveAll(this);
+	}
+#endif // WITH_EDITOR
 }
 
 #if WITH_EDITOR
@@ -71,19 +124,34 @@ void UDMXPixelMappingFixtureGroupItemComponent::PostEditChangeProperty(FProperty
 {
 	Super::PostEditChangeProperty(PropertyChangedEvent);
 
-	FName PropertyName = PropertyChangedEvent.GetPropertyName();
+	const FName PropertyName = PropertyChangedEvent.GetPropertyName();
+	
+	if (PropertyName == GET_MEMBER_NAME_CHECKED(UDMXPixelMappingFixtureGroupItemComponent, ColorSpaceClass))
+	{
+		if (ensureAlwaysMsgf(ColorSpaceClass, TEXT("Color Space Class was set to nullptr. This is not supported.")))
+		{
+			if (ColorSpace)
+			{
+				ColorSpace->GetOnPostEditChangedProperty().RemoveAll(this);
+				ResetDMX();
+			}
+
+			ColorSpace = NewObject<UDMXPixelMappingColorSpace>(this, ColorSpaceClass);
+			ColorSpace->GetOnPostEditChangedProperty().AddUObject(this, &UDMXPixelMappingFixtureGroupItemComponent::OnColorSpacePostEditChangeProperties);
+		}
+	}
 
 	PRAGMA_DISABLE_DEPRECATION_WARNINGS
-	if (PropertyChangedEvent.GetPropertyName() == UDMXPixelMappingOutputComponent::GetPositionXPropertyName() ||
-		PropertyChangedEvent.GetPropertyName() == UDMXPixelMappingOutputComponent::GetPositionYPropertyName())
+	if (PropertyName == UDMXPixelMappingOutputComponent::GetPositionXPropertyName() ||
+		PropertyName == UDMXPixelMappingOutputComponent::GetPositionYPropertyName())
 	{
 		if (ComponentWidget_DEPRECATED.IsValid())
 		{
 			ComponentWidget_DEPRECATED->SetPosition(GetPosition());
 		}
 	}
-	if (PropertyChangedEvent.GetPropertyName() == UDMXPixelMappingOutputComponent::GetSizeXPropertyName() ||
-		PropertyChangedEvent.GetPropertyName() == UDMXPixelMappingOutputComponent::GetSizeYPropertyName())
+	else if (PropertyName == UDMXPixelMappingOutputComponent::GetSizeXPropertyName() ||
+		PropertyName == UDMXPixelMappingOutputComponent::GetSizeYPropertyName())
 	{
 		if (ComponentWidget_DEPRECATED.IsValid())
 		{
@@ -143,9 +211,32 @@ void UDMXPixelMappingFixtureGroupItemComponent::ResetDMX()
 		return;
 	}
 
-	RendererComponent->ResetColorDownsampleBufferPixel(DownsamplePixelIndex);
+	UDMXEntityFixturePatch* FixturePatch = FixturePatchRef.GetFixturePatch();
+	if (!FixturePatch)
+	{
+		return;
+	}
+
+	UDMXLibrary* Library = FixturePatch->GetParentLibrary();
+	if (!Library)
+	{
+		return;
+	}
+	
+	ColorSpace->ResetToBlack();
+	TMap<FDMXAttributeName, float> AttributeToValueMap = ColorSpace->GetAttributeNameToValueMap();
+
+	TMap<int32, uint8> ChannelToValueMap;
+	for (const TTuple<FDMXAttributeName, float>& AttributeValuePair : AttributeToValueMap)
+	{
+		FDMXPixelMappingRuntimeUtils::ConvertNormalizedAttributeValueToChannelValue(FixturePatch, AttributeValuePair.Key, AttributeValuePair.Value, ChannelToValueMap);
+	}
 
-	SendDMX();
+	// Send DMX
+	for (const FDMXOutputPortSharedRef& OutputPort : Library->GetOutputPorts())
+	{
+		OutputPort->SendDMX(FixturePatch->GetUniverseID(), ChannelToValueMap);
+	}
 }
 
 void UDMXPixelMappingFixtureGroupItemComponent::SendDMX()
@@ -153,35 +244,54 @@ void UDMXPixelMappingFixtureGroupItemComponent::SendDMX()
 	SCOPE_CYCLE_COUNTER(STAT_DMXPixelMaping_FixtureGroupItem);
 
 	UDMXEntityFixturePatch* FixturePatch = FixturePatchRef.GetFixturePatch();
+	if (!FixturePatch)
+	{
+		return;
+	}
 
-	if(FixturePatch)
+	UDMXLibrary* Library = FixturePatch->GetParentLibrary();
+	if (!Library)
 	{
-		UDMXPixelMappingRendererComponent* RendererComponent = GetRendererComponent();
-		if (RendererComponent)
-		{
-			TMap<FDMXAttributeName, float> AttributeToValueMap = CreateAttributeValues();
-			
-			// No need to apply matrix modulators
-			for (UDMXModulator* Modulator : Modulators)
-			{
-				Modulator->Modulate(FixturePatch, AttributeToValueMap, AttributeToValueMap);
-			}
-			
-			TMap<int32, uint8> ChannelToValueMap;
-			for (const TTuple<FDMXAttributeName, float>& AttributeValuePair : AttributeToValueMap)
-			{
-				FDMXPixelMappingRuntimeUtils::ConvertNormalizedAttributeValueToChannelValue(FixturePatch, AttributeValuePair.Key, AttributeValuePair.Value, ChannelToValueMap);
-			}
+		return;
+	}
 
-			// Send DMX
-			if (UDMXLibrary* Library = FixturePatch->GetParentLibrary())
-			{
-				for (const FDMXOutputPortSharedRef& OutputPort : Library->GetOutputPorts())
-				{
-					OutputPort->SendDMX(FixturePatch->GetUniverseID(), ChannelToValueMap);
-				}
-			}
-		}
+	UDMXPixelMappingRendererComponent* RendererComponent = GetRendererComponent();
+	if (!RendererComponent)
+	{
+		return;
+	}
+
+	if (!ColorSpace)
+	{
+		return;
+	}
+
+	// Get the color data from the rendered component
+	FLinearColor PixelColor;
+	if (!RendererComponent->GetDownsampleBufferPixel(DownsamplePixelIndex, PixelColor))
+	{
+		return;
+	}
+
+	ColorSpace->SetRGBA(PixelColor);
+	TMap<FDMXAttributeName, float> AttributeToValueMap = ColorSpace->GetAttributeNameToValueMap();
+
+	for (UDMXModulator* Modulator : Modulators)
+	{
+		Modulator->Modulate(FixturePatch, AttributeToValueMap, AttributeToValueMap);
+		// No need to apply Matrix Modulators, this is not a Matrix
+	}
+	
+	TMap<int32, uint8> ChannelToValueMap;
+	for (const TTuple<FDMXAttributeName, float>& AttributeValuePair : AttributeToValueMap)
+	{
+		FDMXPixelMappingRuntimeUtils::ConvertNormalizedAttributeValueToChannelValue(FixturePatch, AttributeValuePair.Key, AttributeValuePair.Value, ChannelToValueMap);
+	}
+
+	// Send DMX
+	for (const FDMXOutputPortSharedRef& OutputPort : Library->GetOutputPorts())
+	{
+		OutputPort->SendDMX(FixturePatch->GetUniverseID(), ChannelToValueMap);
 	}
 }
 
@@ -213,25 +323,8 @@ void UDMXPixelMappingFixtureGroupItemComponent::QueueDownsample()
 	const FVector2D UVCellSize = UVSize / 2.f;
 	constexpr bool bStaticCalculateUV = true;
 
-	FVector4 ExposeFactor{};
-	FIntVector4 InvertFactor{};
-	if (ColorMode == EDMXColorMode::CM_RGB)
-	{
-		ExposeFactor = FVector4(AttributeRExpose ? 1.f : 0.f, AttributeGExpose ? 1.f : 0.f, AttributeBExpose ? 1.f : 0.f, 1.f);
-		InvertFactor = FIntVector4(AttributeRInvert, AttributeGInvert, AttributeBInvert, 0);
-	}
-	else if (ColorMode == EDMXColorMode::CM_Monochrome)
-	{
-		static const FVector4 Expose(1.f, 1.f, 1.f, 1.f);
-		static const FVector4 NoExpose(0.f, 0.f, 0.f, 0.f);
-		ExposeFactor = FVector4(bMonochromeExpose ? Expose : NoExpose);
-		InvertFactor = FIntVector4(bMonochromeInvert, bMonochromeInvert, bMonochromeInvert, 0);
-	}
-
-	FDMXPixelMappingDownsamplePixelParam DownsamplePixelParam
-	{
-		ExposeFactor,
-		InvertFactor,
+	FDMXPixelMappingDownsamplePixelParamsV2 DownsamplePixelParams
+	{ 
 		PixelPosition,
 		UV,
 		UVSize,
@@ -240,7 +333,7 @@ void UDMXPixelMappingFixtureGroupItemComponent::QueueDownsample()
 		bStaticCalculateUV
 	};
 
-	RendererComponent->AddPixelToDownsampleSet(MoveTemp(DownsamplePixelParam));
+	RendererComponent->AddPixelToDownsampleSet(MoveTemp(DownsamplePixelParams));
 }
 
 void UDMXPixelMappingFixtureGroupItemComponent::SetPosition(const FVector2D& NewPosition)
@@ -319,6 +412,17 @@ bool UDMXPixelMappingFixtureGroupItemComponent::CanBeMovedTo(const UDMXPixelMapp
 	return false;
 }
 
+#if WITH_EDITOR
+void UDMXPixelMappingFixtureGroupItemComponent::OnColorSpacePostEditChangeProperties(FPropertyChangedEvent& PropertyChangedEvent)
+{
+	if (PropertyChangedEvent.GetPropertyName() == GET_MEMBER_NAME_CHECKED(FDMXAttributeName, Name))
+	{
+		// Reset DMX when an attribute of the color space changed
+		ResetDMX();
+	}
+}
+#endif // WITH_EDITOR
+
 UDMXPixelMappingRendererComponent* UDMXPixelMappingFixtureGroupItemComponent::UDMXPixelMappingFixtureGroupItemComponent::GetRendererComponent() const
 {
 	return GetParent() ? Cast<UDMXPixelMappingRendererComponent>(GetParent()->GetParent()) : nullptr;
@@ -326,50 +430,15 @@ UDMXPixelMappingRendererComponent* UDMXPixelMappingFixtureGroupItemComponent::UD
 
 TMap<FDMXAttributeName, float> UDMXPixelMappingFixtureGroupItemComponent::CreateAttributeValues() const
 {
-	TMap<FDMXAttributeName, float> AttributeToNormalizedValueMap;
+	// DEPRECATED 5.2
+	TMap<FDMXAttributeName, float> AttributeToValueMap;
 
-	UDMXPixelMappingRendererComponent* RendererComponent = GetRendererComponent();
-	if (RendererComponent)
+	if (ensureMsgf(ColorSpace, TEXT("Unexpected invalid Color Space in Pixel Mapping Fixture Group Item %s."), *GetName()))
 	{
-		// Get the color data from the rendered component
-		FLinearColor PixelColor;
-		if (RendererComponent->GetDownsampleBufferPixel(DownsamplePixelIndex, PixelColor))
-		{
-			if (ColorMode == EDMXColorMode::CM_RGB)
-			{
-				if (AttributeRExpose)
-				{
-					const float AttributeRValue = FMath::Clamp(PixelColor.R, 0.f, 1.f);
-					AttributeToNormalizedValueMap.Add(AttributeR, AttributeRValue);
-				}
-
-				if (AttributeGExpose)
-				{
-					const float AttributeGValue = FMath::Clamp(PixelColor.G, 0.f, 1.f);
-					AttributeToNormalizedValueMap.Add(AttributeG, AttributeGValue);
-				}
-
-				if (AttributeBExpose)
-				{
-					const float AttributeBValue = FMath::Clamp(PixelColor.B, 0.f, 1.f);
-					AttributeToNormalizedValueMap.Add(AttributeB, AttributeBValue);
-				}
-			}
-			else if (ColorMode == EDMXColorMode::CM_Monochrome)
-			{
-				if (bMonochromeExpose)
-				{
-					// https://www.w3.org/TR/AERT/#color-contrast
-					float Intensity = 0.299f * PixelColor.R + 0.587f * PixelColor.G + 0.114f * PixelColor.B;
-					Intensity = FMath::Clamp(Intensity, 0.f, 1.f);
-
-					AttributeToNormalizedValueMap.Add(MonochromeIntensity, Intensity);
-				}
-			}
-		}
+		AttributeToValueMap = ColorSpace->GetAttributeNameToValueMap();
 	}
 
-	return AttributeToNormalizedValueMap;
+	return AttributeToValueMap;
 }
 
 #undef LOCTEXT_NAMESPACE
diff --git a/Engine/Plugins/VirtualProduction/DMX/DMXPixelMapping/Source/DMXPixelMappingRuntime/Private/Components/DMXPixelMappingMatrixCellComponent.cpp b/Engine/Plugins/VirtualProduction/DMX/DMXPixelMapping/Source/DMXPixelMappingRuntime/Private/Components/DMXPixelMappingMatrixCellComponent.cpp
--- a/Engine/Plugins/VirtualProduction/DMX/DMXPixelMapping/Source/DMXPixelMappingRuntime/Private/Components/DMXPixelMappingMatrixCellComponent.cpp
+++ b/Engine/Plugins/VirtualProduction/DMX/DMXPixelMapping/Source/DMXPixelMappingRuntime/Private/Components/DMXPixelMappingMatrixCellComponent.cpp
@@ -3,11 +3,13 @@
 #include "Components/DMXPixelMappingMatrixCellComponent.h"
 
 #include "DMXPixelMappingTypes.h"
+#include "IDMXPixelMappingRenderer.h"
+#include "ColorSpace/DMXPixelMappingColorSpace.h"
+#include "Components/DMXPixelMappingFixtureGroupComponent.h"
 #include "Components/DMXPixelMappingMatrixComponent.h"
 #include "Components/DMXPixelMappingRendererComponent.h"
-#include "IDMXPixelMappingRenderer.h"
 #include "Library/DMXEntityFixturePatch.h"
-#include "TextureResource.h"
+#include "Library/DMXEntityFixtureType.h"
 
 #if WITH_EDITOR
 #include "DMXPixelMappingComponentWidget.h"
@@ -15,6 +17,7 @@
 #include "SDMXPixelMappingComponentLabel.h"
 #endif // WITH_EDITOR
 
+#include "TextureResource.h"
 #include "Engine/Texture.h"
 
 
@@ -22,7 +25,6 @@ DECLARE_CYCLE_STAT(TEXT("Send Matrix Cell"), STAT_DMXPixelMaping_SendMatrixCell,
 
 #define LOCTEXT_NAMESPACE "DMXPixelMappingMatrixPixelComponent"
 
-
 UDMXPixelMappingMatrixCellComponent::UDMXPixelMappingMatrixCellComponent()
 	: DownsamplePixelIndex(0)
 {
@@ -248,26 +250,9 @@ void UDMXPixelMappingMatrixCellComponent::QueueDownsample()
 	const FVector2D UVSize(GetSize().X / TextureSizeX, GetSize().Y / TextureSizeY);
 	const FVector2D UVCellSize = UVSize / 2.f;
 	constexpr bool bStaticCalculateUV = true;
-
-	FVector4 ExposeFactor;
-	FIntVector4 InvertFactor{};
-	if (MatrixComponent->ColorMode == EDMXColorMode::CM_RGB)
-	{
-		ExposeFactor = FVector4(MatrixComponent->AttributeRExpose ? 1.f : 0.f, MatrixComponent->AttributeGExpose ? 1.f : 0.f, MatrixComponent->AttributeBExpose ? 1.f : 0.f, 1.f);
-		InvertFactor = FIntVector4(MatrixComponent->AttributeRInvert, MatrixComponent->AttributeGInvert, MatrixComponent->AttributeBInvert, 0);
-	}
-	else if (MatrixComponent->ColorMode == EDMXColorMode::CM_Monochrome)
-	{
-		static const FVector4 Expose(1.f, 1.f, 1.f, 1.f);
-		static const FVector4 NoExpose(0.f, 0.f, 0.f, 0.f);
-		ExposeFactor = FVector4(MatrixComponent->bMonochromeExpose ? Expose : NoExpose);
-		InvertFactor = FIntVector4(MatrixComponent->bMonochromeInvert, MatrixComponent->bMonochromeInvert, MatrixComponent->bMonochromeInvert, 0);
-	}
 			
-	FDMXPixelMappingDownsamplePixelParam DownsamplePixelParam
-	{
-		ExposeFactor,
-		InvertFactor,
+	FDMXPixelMappingDownsamplePixelParamsV2 DownsamplePixelParam
+	{ 
 		PixelPosition,
 		UV,
 		UVSize,
@@ -315,51 +300,33 @@ TMap<FDMXAttributeName, float> UDMXPixelMappingMatrixCellComponent::CreateAttrib
 {
 	TMap<FDMXAttributeName, float> AttributeToNormalizedValueMap;
 
-	if (UDMXPixelMappingMatrixComponent* ParentMatrix = Cast<UDMXPixelMappingMatrixComponent>(GetParent()))
+	UDMXPixelMappingMatrixComponent* ParentMatrix = Cast<UDMXPixelMappingMatrixComponent>(GetParent());
+	if (!ParentMatrix)
 	{
-		UDMXPixelMappingRendererComponent* RendererComponent = GetRendererComponent();
-		if (RendererComponent)
-		{
-			// Get the color data from the rendered component
-			FLinearColor PixelColor;
-			if (RendererComponent->GetDownsampleBufferPixel(DownsamplePixelIndex, PixelColor))
-			{
-				if (ParentMatrix->ColorMode == EDMXColorMode::CM_RGB)
-				{
-					if (ParentMatrix->AttributeRExpose)
-					{
-						const float AttributeRValue = FMath::Clamp(PixelColor.R, 0.f, 1.f);
-						AttributeToNormalizedValueMap.Add(ParentMatrix->AttributeR, AttributeRValue);
-					}
-
-					if (ParentMatrix->AttributeGExpose)
-					{
-						const float AttributeGValue = FMath::Clamp(PixelColor.G, 0.f, 1.f);
-						AttributeToNormalizedValueMap.Add(ParentMatrix->AttributeG, AttributeGValue);
-					}
-
-					if (ParentMatrix->AttributeBExpose)
-					{
-						const float AttributeBValue = FMath::Clamp(PixelColor.B, 0.f, 1.f);
-						AttributeToNormalizedValueMap.Add(ParentMatrix->AttributeB, AttributeBValue);
-					}
-				}
-				else if (ParentMatrix->ColorMode == EDMXColorMode::CM_Monochrome)
-				{
-					if (ParentMatrix->bMonochromeExpose)
-					{
-						// https://www.w3.org/TR/AERT/#color-contrast
-						float Intensity = 0.299f * PixelColor.R + 0.587f * PixelColor.G + 0.114f * PixelColor.B;
-						Intensity = FMath::Clamp(Intensity, 0.f, 1.f);
-
-						AttributeToNormalizedValueMap.Add(ParentMatrix->MonochromeIntensity, Intensity);
-					}
-				}
-			}
-		}
+		return AttributeToNormalizedValueMap;
 	}
 
-	return AttributeToNormalizedValueMap;
+	UDMXPixelMappingColorSpace* ColorSpace = ParentMatrix->ColorSpace;
+	if (!ColorSpace)
+	{
+		return AttributeToNormalizedValueMap;
+	}
 
+	UDMXPixelMappingRendererComponent* RendererComponent = GetRendererComponent();
+	if (!RendererComponent)
+	{
+		return AttributeToNormalizedValueMap;
+	}
+
+	// Get the color data from the rendered component
+	FLinearColor PixelColor;
+	if (RendererComponent->GetDownsampleBufferPixel(DownsamplePixelIndex, PixelColor))
+	{
+		ColorSpace->SetRGBA(PixelColor);
+		AttributeToNormalizedValueMap = ColorSpace->GetAttributeNameToValueMap();
+	}
+
+	return AttributeToNormalizedValueMap;
 }
+
 #undef LOCTEXT_NAMESPACE
diff --git a/Engine/Plugins/VirtualProduction/DMX/DMXPixelMapping/Source/DMXPixelMappingRuntime/Private/Components/DMXPixelMappingMatrixComponent.cpp b/Engine/Plugins/VirtualProduction/DMX/DMXPixelMapping/Source/DMXPixelMappingRuntime/Private/Components/DMXPixelMappingMatrixComponent.cpp
--- a/Engine/Plugins/VirtualProduction/DMX/DMXPixelMapping/Source/DMXPixelMappingRuntime/Private/Components/DMXPixelMappingMatrixComponent.cpp
+++ b/Engine/Plugins/VirtualProduction/DMX/DMXPixelMapping/Source/DMXPixelMappingRuntime/Private/Components/DMXPixelMappingMatrixComponent.cpp
@@ -3,9 +3,11 @@
 #include "Components/DMXPixelMappingMatrixComponent.h"
 
 #include "DMXPixelMapping.h"
-#include "DMXPixelMappingRuntimeObjectVersion.h"
+#include "DMXPixelMappingMainStreamObjectVersion.h"
+#include "DMXPixelMappingRuntimeCommon.h"
 #include "DMXPixelMappingRuntimeUtils.h"
 #include "DMXPixelMappingTypes.h"
+#include "ColorSpace/DMXPixelMappingColorSpace_RGBCMY.h"
 #include "Components/DMXPixelMappingFixtureGroupComponent.h"
 #include "Components/DMXPixelMappingMatrixCellComponent.h"
 #include "IO/DMXOutputPort.h"
@@ -19,38 +21,54 @@
 #endif // WITH_EDITOR
 
 
-
 #define LOCTEXT_NAMESPACE "DMXPixelMappingMatrixComponent"
 
 UDMXPixelMappingMatrixComponent::UDMXPixelMappingMatrixComponent()
-{	
-	SetSize(FVector2D(32.f, 32.f));
-
-	AttributeR.SetFromName("Red");
-	AttributeG.SetFromName("Green");
-	AttributeB.SetFromName("Blue");
-
-	ColorMode = EDMXColorMode::CM_RGB;
-	AttributeRExpose = true;
-	AttributeGExpose = true;
-	AttributeBExpose = true;
+{
+	ColorSpace = CreateDefaultSubobject<UDMXPixelMappingColorSpace_RGBCMY>("ColorSpace");
 
-	bMonochromeExpose = true;
+	SetSize(FVector2D(32.f, 32.f));
 }
 
 void UDMXPixelMappingMatrixComponent::Serialize(FArchive& Ar)
 {
 	Super::Serialize(Ar);
 
-	Ar.UsingCustomVersion(FDMXPixelMappingRuntimeObjectVersion::GUID);
+#if WITH_EDITOR
+	Ar.UsingCustomVersion(FDMXPixelMappingMainStreamObjectVersion::GUID);
 	if(Ar.IsLoading())
 	{
-		if (Ar.CustomVer(FDMXPixelMappingRuntimeObjectVersion::GUID) < FDMXPixelMappingRuntimeObjectVersion::ChangePixelMappingMatrixComponentToFixturePatchReference)
+		if (Ar.CustomVer(FDMXPixelMappingMainStreamObjectVersion::GUID) < FDMXPixelMappingMainStreamObjectVersion::ChangePixelMappingMatrixComponentToFixturePatchReference)
 		{
 			// Upgrade from custom FixturePatchMatrixRef to FixturePatchRef
 			FixturePatchRef.SetEntity(FixturePatchMatrixRef_DEPRECATED.GetFixturePatch());
 		}
+
+		if (Ar.CustomVer(FDMXPixelMappingMainStreamObjectVersion::GUID) < FDMXPixelMappingMainStreamObjectVersion::UseDMXPixelMappingColorSpace)
+		{		
+			// Upgrade to the DMXPixelMappingColorSpace default subobject
+			UDMXPixelMappingColorSpace_RGBCMY* ColorSpace_RGBCMY = Cast<UDMXPixelMappingColorSpace_RGBCMY>(ColorSpace);
+			if (!ensureMsgf(ColorSpace_RGBCMY, TEXT("Missing default Subobject ColorSpace")))
+			{
+				ColorSpace = NewObject<UDMXPixelMappingColorSpace_RGBCMY>(this, "ColorSpace");
+			}
+
+			if (bMonochromeExpose_DEPRECATED)
+			{
+				ColorSpace_RGBCMY->LuminanceAttribute = MonochromeIntensity_DEPRECATED;
+			}
+			else
+			{
+				ColorSpace_RGBCMY->bSendCyan = AttributeRInvert_DEPRECATED;
+				ColorSpace_RGBCMY->bSendMagenta = AttributeGInvert_DEPRECATED;
+				ColorSpace_RGBCMY->bSendYellow = AttributeBInvert_DEPRECATED;
+				ColorSpace_RGBCMY->RedAttribute = AttributeR_DEPRECATED;
+				ColorSpace_RGBCMY->GreenAttribute = AttributeG_DEPRECATED;
+				ColorSpace_RGBCMY->BlueAttribute = AttributeB_DEPRECATED;
+			}
+		}
 	}
+#endif
 }
 
 void UDMXPixelMappingMatrixComponent::PostLoad()
@@ -75,6 +93,12 @@ void UDMXPixelMappingMatrixComponent::PostLoad()
 			}
 		}
 	}
+
+	// Set the transient Color Space Class property to the initial value
+	if (ensureMsgf(ColorSpace, TEXT("Unexpected invalid Color Space in Pixel Mapping Fixture Group Item %s."), *GetName()))
+	{
+		ColorSpaceClass = ColorSpace->GetClass();
+	}
 }
 
 void UDMXPixelMappingMatrixComponent::PostInitProperties()
@@ -117,8 +141,7 @@ void UDMXPixelMappingMatrixComponent::LogInvalidProperties()
 #if WITH_EDITOR
 void UDMXPixelMappingMatrixComponent::PreEditChange(FProperty* PropertyAboutToChange)
 {
-	// Explicitly recording the state here is required since many properties propagonate to the child cells and cause otherwise malformed transactions
-	Modify();
+	Super::PreEditChange(PropertyAboutToChange);
 
 	if (PropertyAboutToChange)
 	{
@@ -141,8 +164,18 @@ void UDMXPixelMappingMatrixComponent::PostEditChangeProperty(FPropertyChangedEve
 	// Call the parent at the first place
 	Super::PostEditChangeProperty(PropertyChangedEvent);
 
-	if (PropertyChangedEvent.GetPropertyName() == UDMXPixelMappingOutputComponent::GetPositionXPropertyName() ||
-		PropertyChangedEvent.GetPropertyName() == UDMXPixelMappingOutputComponent::GetPositionYPropertyName())
+	const FName PropertyName = PropertyChangedEvent.GetPropertyName();
+
+	if (PropertyName == GET_MEMBER_NAME_CHECKED(UDMXPixelMappingMatrixComponent, ColorSpaceClass))
+	{
+		if (ensureAlwaysMsgf(ColorSpaceClass, TEXT("Color Space Class was set to nullptr. This is not supported.")))
+		{
+			ColorSpace = NewObject<UDMXPixelMappingColorSpace>(this, ColorSpaceClass);
+		}
+	}
+	
+	if (PropertyName == UDMXPixelMappingOutputComponent::GetPositionXPropertyName() ||
+		PropertyName == UDMXPixelMappingOutputComponent::GetPositionYPropertyName())
 	{
 		constexpr bool bUpdateSizeRecursive = false;
 		ForEachChildOfClass<UDMXPixelMappingMatrixCellComponent>([this](UDMXPixelMappingMatrixCellComponent* ChildComponent)
@@ -154,9 +187,10 @@ void UDMXPixelMappingMatrixComponent::PostEditChangeProperty(FPropertyChangedEve
 
 		HandlePositionChanged();
 	}
-	if (PropertyChangedEvent.GetPropertyName() == UDMXPixelMappingOutputComponent::GetSizeXPropertyName() ||
-		PropertyChangedEvent.GetPropertyName() == UDMXPixelMappingOutputComponent::GetSizeYPropertyName() ||
-		PropertyChangedEvent.GetPropertyName() == GET_MEMBER_NAME_CHECKED(UDMXPixelMappingMatrixComponent, CellSize))
+	
+	if (PropertyName == UDMXPixelMappingOutputComponent::GetSizeXPropertyName() ||
+		PropertyName == UDMXPixelMappingOutputComponent::GetSizeYPropertyName() ||
+		PropertyName == GET_MEMBER_NAME_CHECKED(UDMXPixelMappingMatrixComponent, CellSize))
 	{
 		HandleSizeChanged();
 
@@ -164,12 +198,10 @@ void UDMXPixelMappingMatrixComponent::PostEditChangeProperty(FPropertyChangedEve
 		const FVector2D NewSize(CellSize.X * CoordinateGrid.X, CellSize.Y * CoordinateGrid.Y);
 		SetSize(NewSize);
 	}
-	if (PropertyChangedEvent.GetPropertyName() == GET_MEMBER_NAME_CHECKED(UDMXPixelMappingMatrixComponent, FixturePatchRef) ||
-		PropertyChangedEvent.GetPropertyName() == GET_MEMBER_NAME_CHECKED(UDMXPixelMappingMatrixComponent, CoordinateGrid))
-	{
-		HandleMatrixChanged();
-	}
-	else if (PropertyChangedEvent.GetPropertyName() == GET_MEMBER_NAME_CHECKED(FDMXEntityReference, DMXLibrary))
+
+	if (PropertyName == GET_MEMBER_NAME_CHECKED(UDMXPixelMappingMatrixComponent, FixturePatchRef) ||
+		PropertyName == GET_MEMBER_NAME_CHECKED(UDMXPixelMappingMatrixComponent, CoordinateGrid) ||
+		PropertyName == GET_MEMBER_NAME_CHECKED(FDMXEntityReference, DMXLibrary))
 	{
 		HandleMatrixChanged();
 	}
diff --git a/Engine/Plugins/VirtualProduction/DMX/DMXPixelMapping/Source/DMXPixelMappingRuntime/Private/Components/DMXPixelMappingRendererComponent.cpp b/Engine/Plugins/VirtualProduction/DMX/DMXPixelMapping/Source/DMXPixelMappingRuntime/Private/Components/DMXPixelMappingRendererComponent.cpp
--- a/Engine/Plugins/VirtualProduction/DMX/DMXPixelMapping/Source/DMXPixelMappingRuntime/Private/Components/DMXPixelMappingRendererComponent.cpp
+++ b/Engine/Plugins/VirtualProduction/DMX/DMXPixelMapping/Source/DMXPixelMappingRuntime/Private/Components/DMXPixelMappingRendererComponent.cpp
@@ -2,16 +2,18 @@
 
 #include "Components/DMXPixelMappingRendererComponent.h"
 
-#include "DMXPixelMappingRuntimeObjectVersion.h"
+#include "DMXPixelMappingMainStreamObjectVersion.h"
 #include "DMXPixelMappingTypes.h"
 #include "IDMXPixelMappingRenderer.h"
 #include "IDMXPixelMappingRendererModule.h"
+#include "DMXStats.h"
 #include "Components/DMXPixelMappingFixtureGroupComponent.h"
 #include "Components/DMXPixelMappingFixtureGroupItemComponent.h"
 #include "Components/DMXPixelMappingMatrixComponent.h"
 #include "Components/DMXPixelMappingRootComponent.h"
 #include "Components/DMXPixelMappingScreenComponent.h"
 
+#include "Async/Async.h"
 #include "Blueprint/UserWidget.h"
 #include "Engine/TextureRenderTarget2D.h"
 #include "Modulators/DMXModulator.h"
@@ -19,12 +21,16 @@
 #include "UObject/ConstructorHelpers.h"
 #include "Widgets/Layout/SConstraintCanvas.h"
 
+
 #if WITH_EDITOR
 
 #include "LevelEditor.h"
 #endif
 
 
+DECLARE_CYCLE_STAT(TEXT("PixelMapping Render"), STAT_DMXPixelMappingRender, STATGROUP_DMX);
+DECLARE_CYCLE_STAT(TEXT("PixelMapping Send DMX Worker"), STAT_DMXPixelMappingSendDMX, STATGROUP_DMX);
+
 const FIntPoint UDMXPixelMappingRendererComponent::MaxDownsampleBufferTargetSize = FIntPoint(4096);
 const FLinearColor UDMXPixelMappingRendererComponent::ClearTextureColor = FLinearColor::Black;
 
@@ -32,16 +38,22 @@ UDMXPixelMappingRendererComponent::UDMXPixelMappingRendererComponent()
 	: DownsampleBufferTarget(nullptr)
 	, DownsamplePixelCount(0)
 {
+	SetSize(FVector2D(100.f, 100.f));
+
 #if WITH_EDITOR
 	ConstructorHelpers::FObjectFinder<UTexture> DefaultTexture(TEXT("Texture2D'/Engine/VREditor/Devices/Vive/UE4_Logo.UE4_Logo'"), LOAD_NoWarn);
-	// Hit breakpoint instead of preventing the editor to load if not found
-	checkfSlow(DefaultTexture.Succeeded(), TEXT("Failed to load Texture2D'/Engine/VREditor/Devices/Vive/UE4_Logo.UE4_Logo'"));
-	InputTexture = DefaultTexture.Object;
-	RendererType = EDMXPixelMappingRendererType::Texture;
+	if (ensureAlwaysMsgf(DefaultTexture.Succeeded(), TEXT("Failed to load Texture2D'/Engine/VREditor/Devices/Vive/UE4_Logo.UE4_Logo'")))
+	{
+		InputTexture = DefaultTexture.Object;
+		RendererType = EDMXPixelMappingRendererType::Texture;
+
+		if (FTextureResource* Resource = InputTexture->GetResource())
+		{
+			SetSize(FVector2D(Resource->GetSizeX(), Resource->GetSizeY()));
+		}
+	}
 #endif
 	
-	SetSize(FVector2D(100.f, 100.f));
-
 	Brightness = 1.0f;
 
 #if WITH_EDITOR
@@ -71,10 +83,10 @@ void UDMXPixelMappingRendererComponent::Serialize(FArchive& Ar)
 {
 	Super::Serialize(Ar);
 
-	Ar.UsingCustomVersion(FDMXPixelMappingRuntimeObjectVersion::GUID);
+	Ar.UsingCustomVersion(FDMXPixelMappingMainStreamObjectVersion::GUID);
 	if (Ar.IsLoading())
 	{
-		if (Ar.CustomVer(FDMXPixelMappingRuntimeObjectVersion::GUID) < FDMXPixelMappingRuntimeObjectVersion::LockRendererComponentsThatUseTextureInDesigner)
+		if (Ar.CustomVer(FDMXPixelMappingMainStreamObjectVersion::GUID) < FDMXPixelMappingMainStreamObjectVersion::LockRendererComponentsThatUseTextureInDesigner)
 		{
 			if (RendererType == EDMXPixelMappingRendererType::Texture)
 			{
@@ -173,6 +185,8 @@ void UDMXPixelMappingRendererComponent::PostEditChangeChainProperty(FPropertyCha
 			Renderer->SetBrightness(Brightness);
 		}
 	}
+
+	Initialize();
 }
 #endif // WITH_EDITOR
 
@@ -319,26 +333,36 @@ void UDMXPixelMappingRendererComponent::ResizePreviewRenderTarget(uint32 InSizeX
 void UDMXPixelMappingRendererComponent::ResetDMX()
 {
 	ForEachChild([&](UDMXPixelMappingBaseComponent* InComponent)
-	{
-		if (UDMXPixelMappingOutputComponent * Component = Cast<UDMXPixelMappingOutputComponent>(InComponent))
 		{
-			Component->ResetDMX();
-		}
-	}, false);
+			if (UDMXPixelMappingOutputComponent* Component = Cast<UDMXPixelMappingOutputComponent>(InComponent))
+			{
+				Component->ResetDMX();
+			}
+		}, false);
 }
 
 void UDMXPixelMappingRendererComponent::SendDMX()
 {
-	ForEachChild([&](UDMXPixelMappingBaseComponent* InComponent) {
-		if (UDMXPixelMappingOutputComponent* Component = Cast<UDMXPixelMappingOutputComponent>(InComponent))
+	ForEachChild([&](UDMXPixelMappingBaseComponent* InComponent)
 		{
-			Component->SendDMX();
-		}
-	}, false);
+			if (UDMXPixelMappingOutputComponent* Component = Cast<UDMXPixelMappingOutputComponent>(InComponent))
+			{
+				SCOPE_CYCLE_COUNTER(STAT_DMXPixelMappingSendDMX);
+
+				Component->SendDMX();
+			}
+		}, false);
 }
 
 void UDMXPixelMappingRendererComponent::Render()
 {
+	SCOPE_CYCLE_COUNTER(STAT_DMXPixelMappingRender);
+
+	if (!PixelMappingRenderer.IsValid())
+	{
+		return;
+	}
+
 	// 1. Get downsample input texture
 	UTexture* DownsampleInputTexture = GetRendererInputTexture();
 	if (!DownsampleInputTexture)
@@ -371,7 +395,7 @@ void UDMXPixelMappingRendererComponent::Render()
 	}
 
 	// 7. Downsample all pixels
-	GetRenderer()->DownsampleRender(
+	PixelMappingRenderer->DownsampleRender(
 		DownsampleInputTexture->GetResource(),
 		DownsampleBufferTarget->GetResource(),
 		DownsampleBufferTarget->GameThread_GetRenderTargetResource(),
@@ -453,6 +477,13 @@ void UDMXPixelMappingRendererComponent::Initialize()
 	}
 #endif
 
+	if (InputTexture)
+	{
+		InputTexture->SetForceMipLevelsToBeResident(30.0f);
+		InputTexture->WaitForStreaming();
+	}
+
+
 	if (UserWidget == nullptr && InputWidget != nullptr)
 	{
 		UserWidget = CreateWidget(GetWorld(), InputWidget);
@@ -501,8 +532,6 @@ UTextureRenderTarget2D* UDMXPixelMappingRendererComponent::CreateRenderTarget(co
 
 void UDMXPixelMappingRendererComponent::RendererInputTexture()
 {
-	Initialize();
-
 	switch (RendererType)
 	{
 	case(EDMXPixelMappingRendererType::Texture):
@@ -568,6 +597,8 @@ bool UDMXPixelMappingRendererComponent::CanBeMovedTo(const UDMXPixelMappingBaseC
 
 int32 UDMXPixelMappingRendererComponent::GetTotalDownsamplePixelCount()
 {
+	FScopeLock ScopeLock(&DownsampleBufferCS);
+
 	// Reset pixel counter
 	DownsamplePixelCount = 0;
 
@@ -635,13 +666,22 @@ void UDMXPixelMappingRendererComponent::CreateOrUpdateDownsampleBufferTarget()
 	}
 }
 
-void UDMXPixelMappingRendererComponent::AddPixelToDownsampleSet(FDMXPixelMappingDownsamplePixelParam&& InDownsamplePixelParam)
+void UDMXPixelMappingRendererComponent::AddPixelToDownsampleSet(FDMXPixelMappingDownsamplePixelParamsV2&& InDownsamplePixelParam)
 {
+	const FScopeLock Lock(&DownsampleBufferCS);
 	DownsamplePixelParams.Emplace(InDownsamplePixelParam);
 }
 
+int32 UDMXPixelMappingRendererComponent::GetDownsamplePixelNum()
+{
+	const FScopeLock Lock(&DownsampleBufferCS);
+	return DownsamplePixelParams.Num();
+}
+
 bool UDMXPixelMappingRendererComponent::IsPixelRangeValid(const int32 InDownsamplePixelIndexStart, const int32 InDownsamplePixelIndexEnd) const
 {
+	FScopeLock ScopeLock(&DownsampleBufferCS);
+
 	if (InDownsamplePixelIndexEnd >= InDownsamplePixelIndexStart &&
 		DownsampleBuffer.IsValidIndex(InDownsamplePixelIndexStart) &&
 		DownsampleBuffer.IsValidIndex(InDownsamplePixelIndexEnd))
@@ -656,6 +696,16 @@ void UDMXPixelMappingRendererComponent::SetDownsampleBuffer(TArray<FLinearColor>
 {
 	check(IsInRenderingThread());
 
+	if (!bWasEverRendered)
+	{
+		for (int32 PixelIndex = 0; PixelIndex < GetTotalDownsamplePixelCount(); PixelIndex++)
+		{
+			ResetColorDownsampleBufferPixel(PixelIndex);
+		}
+
+		bWasEverRendered = true;
+	}
+
 	FScopeLock ScopeLock(&DownsampleBufferCS);
 	DownsampleBuffer = MoveTemp(InDownsampleBuffer);
 }
@@ -664,13 +714,14 @@ bool UDMXPixelMappingRendererComponent::GetDownsampleBufferPixel(const int32 InD
 {
 	FScopeLock ScopeLock(&DownsampleBufferCS);
 
-	if (DownsampleBuffer.IsValidIndex(InDownsamplePixelIndex))
+
+	if (!DownsampleBuffer.IsValidIndex(InDownsamplePixelIndex))
 	{
-		OutLinearColor = DownsampleBuffer[InDownsamplePixelIndex];
-		return true;
+		return false;
 	}
 
-	return false;
+	OutLinearColor = DownsampleBuffer[InDownsamplePixelIndex];
+	return true;
 }
 
 bool UDMXPixelMappingRendererComponent::GetDownsampleBufferPixels(const int32 InDownsamplePixelIndexStart, const int32 InDownsamplePixelIndexEnd, TArray<FLinearColor>& OutLinearColors)
diff --git a/Engine/Plugins/VirtualProduction/DMX/DMXPixelMapping/Source/DMXPixelMappingRuntime/Private/Components/DMXPixelMappingRootComponent.cpp b/Engine/Plugins/VirtualProduction/DMX/DMXPixelMapping/Source/DMXPixelMappingRuntime/Private/Components/DMXPixelMappingRootComponent.cpp
--- a/Engine/Plugins/VirtualProduction/DMX/DMXPixelMapping/Source/DMXPixelMappingRuntime/Private/Components/DMXPixelMappingRootComponent.cpp
+++ b/Engine/Plugins/VirtualProduction/DMX/DMXPixelMapping/Source/DMXPixelMappingRuntime/Private/Components/DMXPixelMappingRootComponent.cpp
@@ -4,6 +4,18 @@
 #include "Components/DMXPixelMappingRendererComponent.h"
 
 
+void UDMXPixelMappingRootComponent::PostLoad()
+{
+	Super::PostLoad();
+
+	ForEachChildOfClass<UDMXPixelMappingRendererComponent>([this](UDMXPixelMappingRendererComponent* RendererComponent)
+		{
+			CachedRendererComponentsByName.Add(RendererComponent->GetFName(), RendererComponent);
+		}, false);
+
+	GetOnComponentRenamed().AddUObject(this, &UDMXPixelMappingRootComponent::OnComponentRenamed);
+}
+
 void UDMXPixelMappingRootComponent::BeginDestroy()
 {
 	Super::BeginDestroy();
@@ -18,6 +30,23 @@ const FName& UDMXPixelMappingRootComponent::GetNamePrefix()
 	return NamePrefix;
 }
 
+void UDMXPixelMappingRootComponent::AddChild(UDMXPixelMappingBaseComponent* InComponent)
+{
+	Super::AddChild(InComponent);
+
+	if (UDMXPixelMappingRendererComponent* RendererComponent = Cast<UDMXPixelMappingRendererComponent>(InComponent))
+	{
+		CachedRendererComponentsByName.FindOrAdd(InComponent->GetFName(), RendererComponent);
+	}
+}
+
+void UDMXPixelMappingRootComponent::RemoveChild(UDMXPixelMappingBaseComponent* InComponent)
+{
+	Super::RemoveChild(InComponent);
+
+	CachedRendererComponentsByName.Remove(InComponent->GetFName());
+}
+
 void UDMXPixelMappingRootComponent::ResetDMX()
 {
 	ForEachChild([&](UDMXPixelMappingBaseComponent* InComponent)
@@ -61,3 +90,21 @@ bool UDMXPixelMappingRootComponent::CanBeMovedTo(const UDMXPixelMappingBaseCompo
 {
 	return false;
 }
+
+UDMXPixelMappingRendererComponent* UDMXPixelMappingRootComponent::FindRendererComponentByName(const FName& Name) const
+{
+	if (const TObjectPtr<UDMXPixelMappingRendererComponent>* RendererComponentPtr = CachedRendererComponentsByName.Find(Name))
+	{
+		return *RendererComponentPtr;
+	}
+
+	return nullptr;
+}
+
+void UDMXPixelMappingRootComponent::OnComponentRenamed(UDMXPixelMapping* PixelMapping, UDMXPixelMappingBaseComponent* RenamedComponent, UObject* OldOuter, const FName OldName)
+{
+	ForEachChildOfClass<UDMXPixelMappingRendererComponent>([this](UDMXPixelMappingRendererComponent* RendererComponent)
+		{
+			CachedRendererComponentsByName.Add(RendererComponent->GetFName(), RendererComponent);
+		}, false);
+}
diff --git a/Engine/Plugins/VirtualProduction/DMX/DMXPixelMapping/Source/DMXPixelMappingRuntime/Private/Components/DMXPixelMappingScreenComponent.cpp b/Engine/Plugins/VirtualProduction/DMX/DMXPixelMapping/Source/DMXPixelMappingRuntime/Private/Components/DMXPixelMappingScreenComponent.cpp
--- a/Engine/Plugins/VirtualProduction/DMX/DMXPixelMapping/Source/DMXPixelMappingRuntime/Private/Components/DMXPixelMappingScreenComponent.cpp
+++ b/Engine/Plugins/VirtualProduction/DMX/DMXPixelMapping/Source/DMXPixelMappingRuntime/Private/Components/DMXPixelMappingScreenComponent.cpp
@@ -410,10 +410,8 @@ void UDMXPixelMappingScreenComponent::QueueDownsample()
             const FIntPoint PixelPosition = RendererComponent->GetPixelPosition(InXYIndex + PixelDownsamplePositionRange.Key);
             const FVector2D UV = FVector2D((GetPosition().X + SizePixel.X * XIndex) / TextureSizeX, (GetPosition().Y + SizePixel.Y * YIndex) / TextureSizeY);
 
-            FDMXPixelMappingDownsamplePixelParam DownsamplePixelParam
+            FDMXPixelMappingDownsamplePixelParamsV2 DownsamplePixelParam
             {
-                PixelFactor,
-                InvertPixel,
                 PixelPosition,
                 UV,
                 UVSize,
diff --git a/Engine/Plugins/VirtualProduction/DMX/DMXPixelMapping/Source/DMXPixelMappingRuntime/Private/DMXPixelMappingMainStreamObjectVersion.cpp b/Engine/Plugins/VirtualProduction/DMX/DMXPixelMapping/Source/DMXPixelMappingRuntime/Private/DMXPixelMappingMainStreamObjectVersion.cpp
new file mode 100644
--- /dev/null
+++ b/Engine/Plugins/VirtualProduction/DMX/DMXPixelMapping/Source/DMXPixelMappingRuntime/Private/DMXPixelMappingMainStreamObjectVersion.cpp
@@ -0,0 +1,11 @@
+// Copyright Epic Games, Inc. All Rights Reserved.
+
+#include "DMXPixelMappingMainStreamObjectVersion.h"
+
+#include "Serialization/CustomVersion.h"
+
+
+const FGuid FDMXPixelMappingMainStreamObjectVersion::GUID(0x75A7B683, 0x4F4C7E28, 0xDDE54885, 0x85323940);
+
+// Register the custom version with core
+FCustomVersionRegistration GRegisterDMXPixelMappingMainStreamCustomVersion(FDMXPixelMappingMainStreamObjectVersion::GUID, FDMXPixelMappingMainStreamObjectVersion::LatestVersion, TEXT("DMXMainStreamObjectVer"));
diff --git a/Engine/Plugins/VirtualProduction/DMX/DMXPixelMapping/Source/DMXPixelMappingRuntime/Private/DMXPixelMappingMainStreamObjectVersion.h b/Engine/Plugins/VirtualProduction/DMX/DMXPixelMapping/Source/DMXPixelMappingRuntime/Private/DMXPixelMappingMainStreamObjectVersion.h
new file mode 100644
--- /dev/null
+++ b/Engine/Plugins/VirtualProduction/DMX/DMXPixelMapping/Source/DMXPixelMappingRuntime/Private/DMXPixelMappingMainStreamObjectVersion.h
@@ -0,0 +1,38 @@
+// Copyright Epic Games, Inc. All Rights Reserved.
+#pragma once
+
+#include "CoreTypes.h"
+#include "Misc/Guid.h"
+
+
+// Custom serialization version for changes to DMX Pixel Mapping Objects in the Main Stream
+struct FDMXPixelMappingMainStreamObjectVersion
+{
+	enum Type
+	{
+		// Roughly corresponds to 4.26
+		BeforeCustomVersionWasAdded = 0,
+
+		// Update DMXPixelMappingMatrixComponent to no longer use a custom MatrixFixturePatchRef, but the default EntityFixturePatchRef one.
+		ChangePixelMappingMatrixComponentToFixturePatchReference,
+
+		// Update DMXPixelMappingBaseComponent to store the parent as a weak ptr
+		UseWeakPtrForPixelMappingComponentParent,
+
+		// Update DMXPixelMappingRendererComponent to lock those components that use a texture in designer to trigger the edit condition of the size property
+		LockRendererComponentsThatUseTextureInDesigner,
+
+		// Update DMXPixelMappingFixtureGroupItemComponent and DMXPixelMappingMatrixCellComponent to use DMXPixelMappingColorSpace
+		UseDMXPixelMappingColorSpace,
+
+		// -----<new versions can be added above this line>-------------------------------------------------
+		VersionPlusOne,
+		LatestVersion = VersionPlusOne - 1
+	};
+
+	// The GUID for this custom version number
+	const static FGuid GUID;
+
+private:
+	FDMXPixelMappingMainStreamObjectVersion() {}
+};
diff --git a/Engine/Plugins/VirtualProduction/DMX/DMXPixelMapping/Source/DMXPixelMappingRuntime/Private/DMXPixelMappingRuntimeObjectVersion.cpp b/Engine/Plugins/VirtualProduction/DMX/DMXPixelMapping/Source/DMXPixelMappingRuntime/Private/DMXPixelMappingRuntimeObjectVersion.cpp
deleted file mode 100644
diff --git a/Engine/Plugins/VirtualProduction/DMX/DMXPixelMapping/Source/DMXPixelMappingRuntime/Private/DMXPixelMappingRuntimeObjectVersion.h b/Engine/Plugins/VirtualProduction/DMX/DMXPixelMapping/Source/DMXPixelMappingRuntime/Private/DMXPixelMappingRuntimeObjectVersion.h
deleted file mode 100644
diff --git a/Engine/Plugins/VirtualProduction/DMX/DMXPixelMapping/Source/DMXPixelMappingRuntime/Public/ColorSpace/DMXPixelMappingColorSpace.h b/Engine/Plugins/VirtualProduction/DMX/DMXPixelMapping/Source/DMXPixelMappingRuntime/Public/ColorSpace/DMXPixelMappingColorSpace.h
new file mode 100644
--- /dev/null
+++ b/Engine/Plugins/VirtualProduction/DMX/DMXPixelMapping/Source/DMXPixelMappingRuntime/Public/ColorSpace/DMXPixelMappingColorSpace.h
@@ -0,0 +1,67 @@
+// Copyright Epic Games, Inc. All Rights Reserved.
+
+#pragma once
+
+#include "DMXAttribute.h"
+
+#include "CoreMinimal.h"
+#include "UObject/Object.h"
+
+#include "DMXPixelMappingColorSpace.generated.h"
+
+
+/** 
+ * Base class for Pixel Mapping Color Spaces. 
+ * Note, implementations must be thread-safe. 
+ */
+UCLASS(Abstract)
+class DMXPIXELMAPPINGRUNTIME_API UDMXPixelMappingColorSpace
+	: public UObject
+{
+	GENERATED_BODY()
+
+#if WITH_EDITOR
+	DECLARE_MULTICAST_DELEGATE_OneParam(FDMXPixelMappingColorSpacePostEditChangePropertyDelegate, FPropertyChangedEvent& /** PropertyChangedEvent */)
+#endif // WITH_EDITOR
+
+public:
+	/** Sets the RGBA value to be converted to the implemented Color Space */
+	virtual void SetRGBA(const FLinearColor& InColor) PURE_VIRTUAL(UDMXPixelMappingColorSpace::SetRGBA, return;);
+
+	/** Resets the color to black, and any luminance to zero */
+	void ResetToBlack();
+
+	/** Returns the Attribute Name to Value Map of the Color currently set */
+	FORCEINLINE const TMap<FDMXAttributeName, float>& GetAttributeNameToValueMap() { return AttributeNameToValueMap; }
+
+#if WITH_EDITOR
+	/** Returns a delegate broadcast when a property changed. Useful to handle attribute name changes. */
+	FDMXPixelMappingColorSpacePostEditChangePropertyDelegate& GetOnPostEditChangedProperty() { return OnPostEditChangeProperty; }
+#endif // WITH_EDITOR 
+
+	/** When checked, uses Unreal Working Color Space for input conversion, otherwise, assumes the input is sRGB */
+	UPROPERTY(EditAnywhere, Category = "Color Space", Meta = (DisplayPriority = -1))
+	bool bUseWorkingColorSpaceForInput = true;
+
+protected:
+	//~ Begin UObject interface
+#if WITH_EDITOR
+	virtual void PostEditChangeProperty(FPropertyChangedEvent& PropertyChangedEvent) override;
+#endif // WITH_EDITOR
+	//~ End UObject interface
+
+	/** Sets an Attribute Value */
+	void SetAttributeValue(const FDMXAttributeName& AttributeName, float Value);
+
+	/** Clears cached Attribtues */
+	void ClearCachedAttributeValues();
+
+private:
+	/** Buffer holding results of the conversion */
+	TMap<FDMXAttributeName, float> AttributeNameToValueMap;
+
+#if WITH_EDITOR
+	/** Event raised Post Edit Change Property */
+	FDMXPixelMappingColorSpacePostEditChangePropertyDelegate OnPostEditChangeProperty;
+#endif // WITH_EDITOR
+};
diff --git a/Engine/Plugins/VirtualProduction/DMX/DMXPixelMapping/Source/DMXPixelMappingRuntime/Public/Components/DMXPixelMappingBaseComponent.h b/Engine/Plugins/VirtualProduction/DMX/DMXPixelMapping/Source/DMXPixelMappingRuntime/Public/Components/DMXPixelMappingBaseComponent.h
--- a/Engine/Plugins/VirtualProduction/DMX/DMXPixelMapping/Source/DMXPixelMappingRuntime/Public/Components/DMXPixelMappingBaseComponent.h
+++ b/Engine/Plugins/VirtualProduction/DMX/DMXPixelMapping/Source/DMXPixelMappingRuntime/Public/Components/DMXPixelMappingBaseComponent.h
@@ -3,9 +3,11 @@
 #pragma once
 
 #include "DMXPixelMappingRuntimeCommon.h"
+
+#include "Tickable.h"
 #include "UObject/Object.h"
 #include "Stats/Stats.h"
-#include "Tickable.h"
+
 #include "DMXPixelMappingBaseComponent.generated.h"
 
 
@@ -16,9 +18,9 @@ class UDMXPixelMappingRendererComponent;
 
 
 /**
- * Base class for all DMX Pixel Mapping components
+ * Base class for all DMX Pixel Mapping components. 
  */
-UCLASS(BlueprintType, Blueprintable, Abstract)
+UCLASS(BlueprintType, NotBlueprintable, Abstract)
 class DMXPIXELMAPPINGRUNTIME_API UDMXPixelMappingBaseComponent
 	: public UObject
 	, public FTickableGameObject
@@ -42,16 +44,15 @@ public:
 	/** Gets an Event broadcast when a component was renamed */
 	static FDMXPixelMappingOnComponentRenamed& GetOnComponentRenamed();
 	
-protected:
 	//~ Begin UObject interface
-	virtual void Serialize(FArchive& Ar) override;
+protected:
 	virtual void PostRename(UObject* OldOuter, const FName OldName) override;
+	virtual void Serialize(FArchive& Ar) override;
 #if WITH_EDITOR
 	virtual void PostEditUndo() override;
 #endif // WITH_EDITOR
 	//~ End UObject interface
 
-
 public:
 	/**
 	 * Should log properties that were changed in underlying fixture patch or fixture type
@@ -190,7 +191,7 @@ public:
 	UFUNCTION(BlueprintCallable, Category = "DMX|PixelMapping")
 	virtual void ResetDMX() {};
 
-	/** Send DMX values of this component and all children */
+	/** Send DMX values of this component and all children. */
 	UFUNCTION(BlueprintCallable, Category = "DMX|PixelMapping")
 	virtual void SendDMX() {};
 
diff --git a/Engine/Plugins/VirtualProduction/DMX/DMXPixelMapping/Source/DMXPixelMappingRuntime/Public/Components/DMXPixelMappingFixtureGroupItemComponent.h b/Engine/Plugins/VirtualProduction/DMX/DMXPixelMapping/Source/DMXPixelMappingRuntime/Public/Components/DMXPixelMappingFixtureGroupItemComponent.h
--- a/Engine/Plugins/VirtualProduction/DMX/DMXPixelMapping/Source/DMXPixelMappingRuntime/Public/Components/DMXPixelMappingFixtureGroupItemComponent.h
+++ b/Engine/Plugins/VirtualProduction/DMX/DMXPixelMapping/Source/DMXPixelMappingRuntime/Public/Components/DMXPixelMappingFixtureGroupItemComponent.h
@@ -12,33 +12,41 @@
 
 #include "DMXPixelMappingFixtureGroupItemComponent.generated.h"
 
+enum class EDMXColorMode : uint8;
 struct FDMXAttributeName;
 class UDMXModulator;
-enum class EDMXColorMode : uint8;
+class UDMXPixelMappingColorSpace;
 
+class FEditPropertyChain;
 class STextBlock;
 class UTextureRenderTarget2D;
 
 
-/**
- * Fixture Item pixel component
+/** 
+ * A component that holds a single Fixture Patch in the Pixel Mapping, and actually sends DMX.
  */
 UCLASS(AutoExpandCategories = ("Output Settings"))
 class DMXPIXELMAPPINGRUNTIME_API UDMXPixelMappingFixtureGroupItemComponent
 	: public UDMXPixelMappingOutputDMXComponent
 {
 	GENERATED_BODY()
+
 public:
 	/** Default Constructor */
 	UDMXPixelMappingFixtureGroupItemComponent();
 
 protected:
+	//~ Begin UObject interface
+	virtual void Serialize(FArchive& Ar) override;
 	virtual void PostLoad() override;
+	virtual void BeginDestroy() override;
 #if WITH_EDITOR
 	virtual void PostEditChangeProperty(FPropertyChangedEvent& PropertyChangedEvent) override;
 	virtual void PostEditChangeChainProperty(FPropertyChangedChainEvent& PropertyChangedChainEvent) override;
 #endif // WITH_EDITOR
 
+	//~ End UObject interface
+
 public:
 	//~ Begin UDMXPixelMappingBaseComponent implementation
 	virtual const FName& GetNamePrefix() override;
@@ -58,7 +66,6 @@ public:
 	virtual void SetSize(const FVector2D& NewSize) override;
 	//~ End UDMXPixelMappingOutputComponent implementation
 
-public:
 	//~ Begin UDMXPixelMappingOutputDMXComponent implementation
 	virtual void RenderWithInputAndSendDMX() override;
 	//~ End UDMXPixelMappingOutputDMXComponent implementation
@@ -66,63 +73,84 @@ public:
 	/** Check if a Component can be moved under another one (used for copy/move/duplicate) */
 	virtual bool CanBeMovedTo(const UDMXPixelMappingBaseComponent* Component) const override;
 
-private:
-	/** Helper that returns the renderer component this component belongs to */
-	UDMXPixelMappingRendererComponent* GetRendererComponent() const;
-
-public:
 	UPROPERTY(VisibleAnywhere, BlueprintReadOnly, Category = "Selected Patch")
 	FDMXEntityFixturePatchRef FixturePatchRef;
 
-	UPROPERTY(EditAnywhere, BlueprintReadOnly, Category = "Output Settings")
-	EDMXColorMode ColorMode;
+	/** Sets which color space Pixel Mapping sends */
+	UPROPERTY(Transient, EditAnywhere, NoClear, Category = "Color Space", Meta = (DisplayPriority = 2, DisplayName = "Output Mode", ShowDisplayNames))
+	TSubclassOf<UDMXPixelMappingColorSpace> ColorSpaceClass;
 
-	UPROPERTY(EditAnywhere, BlueprintReadOnly, Category = "Output Settings", meta = (DisplayName = "R"))
-	bool AttributeRExpose;
+	/** The Color Space currently in use */
+	UPROPERTY(VisibleAnywhere, Instanced, Category = "Color Space")
+	TObjectPtr<UDMXPixelMappingColorSpace> ColorSpace;
 
-	UPROPERTY(EditAnywhere, BlueprintReadOnly, Category = "Output Settings", meta = (DisplayName = "G"))
-	bool AttributeGExpose;
+	/** Modulators applied to the output before sending DMX */
+	UPROPERTY(Transient, EditAnywhere, BlueprintReadOnly, Transient, Category = "Output Settings", meta = (DisplayName = "Output Modulators"))
+	TArray<TSubclassOf<UDMXModulator>> ModulatorClasses;
 
-	UPROPERTY(EditAnywhere, BlueprintReadOnly, Category = "Output Settings", meta = (DisplayName = "B"))
-	bool AttributeBExpose;
+	/** Modulator instances applied to this component */
+	UPROPERTY()
+	TArray<TObjectPtr<UDMXModulator>> Modulators;
 
-	UPROPERTY(EditAnywhere, BlueprintReadOnly, Category = "Output Settings", meta = (DisplayName = "Expose"))
-	bool bMonochromeExpose;
+	/** Index of the cell pixel in downsample target buffer */
+	int32 DownsamplePixelIndex;
 
-	UPROPERTY(EditAnywhere, BlueprintReadOnly, Category = "Output Settings", meta = (DisplayName = "Invert R"))
-	bool AttributeRInvert;
+private:
+#if WITH_EDITOR
+	/** Called post edit change properties of the color space member */
+	void OnColorSpacePostEditChangeProperties(FPropertyChangedEvent& PropertyChangedEvent);
+#endif // WITH_EDITOR
 
-	UPROPERTY(EditAnywhere, BlueprintReadOnly, Category = "Output Settings", meta = (DisplayName = "Invert G"))
-	bool AttributeGInvert;
+	/** Helper that returns the renderer component this component belongs to */
+	UDMXPixelMappingRendererComponent* GetRendererComponent() const;
 
-	UPROPERTY(EditAnywhere, BlueprintReadOnly, Category = "Output Settings", meta = (DisplayName = "Invert B"))
-	bool AttributeBInvert;
 
-	UPROPERTY(EditAnywhere, BlueprintReadOnly, Category = "Output Settings", meta = (DisplayName = "Invert"))
-	bool bMonochromeInvert;
+	//////////////////////
+	// Deprecated Members
+	
+public:
+	/** DEPRECATED 5.2 */
+	UE_DEPRECATED(5.2, "Instead use UDMXPixelMappingColorSpace::Update and UDMXPixelMappingColorSpace::GetAttributeNameToValueMap using the ColorSpace member of this class.")
+	TMap<FDMXAttributeName, float> CreateAttributeValues() const;
 
-	UPROPERTY(EditAnywhere, BlueprintReadOnly, Category = "Output Settings", meta = (DisplayName = "R Attribute"))
-	FDMXAttributeName AttributeR;
+#if WITH_EDITORONLY_DATA
+	UPROPERTY(Meta = (DeprecatedProperty, DeprecationMessage = "Deprecated in favor of DMXPixelMappingColorSpace. See ColorSpace member."))
+	EDMXColorMode ColorMode_DEPRECATED;
 
-	UPROPERTY(EditAnywhere, BlueprintReadOnly, Category = "Output Settings", meta = (DisplayName = "G Attribute"))
-	FDMXAttributeName AttributeG;
+	UPROPERTY(Meta = (DeprecatedProperty, DeprecationMessage = "Deprecated in favor of DMXPixelMappingColorSpace. See ColorSpace member."))
+	bool AttributeRExpose_DEPRECATED;
 
-	UPROPERTY(EditAnywhere, BlueprintReadOnly, Category = "Output Settings", meta = (DisplayName = "B Attribute"))
-	FDMXAttributeName AttributeB;
+	UPROPERTY(Meta = (DeprecatedProperty, DeprecationMessage = "Deprecated in favor of DMXPixelMappingColorSpace. See ColorSpace member."))
+	bool AttributeGExpose_DEPRECATED;
 
-	UPROPERTY(EditAnywhere, BlueprintReadOnly, Category = "Output Settings", meta = (DisplayName = "Intensity Attribute"))
-	FDMXAttributeName MonochromeIntensity;
+	UPROPERTY(Meta = (DeprecatedProperty, DeprecationMessage = "Deprecated in favor of DMXPixelMappingColorSpace. See ColorSpace member."))
+	bool AttributeBExpose_DEPRECATED;
 
-	/** Modulators applied to the output before sending DMX */
-	UPROPERTY(EditAnywhere, BlueprintReadOnly, Transient, Category = "Output Settings", meta = (DisplayName = "Output Modulators"))
-	TArray<TSubclassOf<UDMXModulator>> ModulatorClasses;
+	UPROPERTY(Meta = (DeprecatedProperty, DeprecationMessage = "Deprecated in favor of DMXPixelMappingColorSpace. See ColorSpace member."))
+	bool bMonochromeExpose_DEPRECATED;
 
-	UPROPERTY()
-	TArray<TObjectPtr<UDMXModulator>> Modulators;
+	UPROPERTY(Meta = (DeprecatedProperty, DeprecationMessage = "Deprecated in favor of DMXPixelMappingColorSpace. See ColorSpace member."))
+	bool AttributeRInvert_DEPRECATED;
 
-	/** Index of the cell pixel in downsample target buffer */
-	int32 DownsamplePixelIndex;
+	UPROPERTY(Meta = (DeprecatedProperty, DeprecationMessage = "Deprecated in favor of DMXPixelMappingColorSpace. See ColorSpace member."))
+	bool AttributeGInvert_DEPRECATED;
 
-	/** Creates attribute values from current data */
-	TMap<FDMXAttributeName, float> CreateAttributeValues() const;
+	UPROPERTY(Meta = (DeprecatedProperty, DeprecationMessage = "Deprecated in favor of DMXPixelMappingColorSpace. See ColorSpace member."))
+	bool AttributeBInvert_DEPRECATED;
+
+	UPROPERTY(Meta = (DeprecatedProperty, DeprecationMessage = "Deprecated in favor of DMXPixelMappingColorSpace. See ColorSpace member."))
+	bool bMonochromeInvert_DEPRECATED;
+
+	UPROPERTY(Meta = (DeprecatedProperty, DeprecationMessage = "Deprecated in favor of DMXPixelMappingColorSpace. See ColorSpace member."))
+	FDMXAttributeName AttributeR_DEPRECATED;
+
+	UPROPERTY(Meta = (DeprecatedProperty, DeprecationMessage = "Deprecated in favor of DMXPixelMappingColorSpace. See ColorSpace member."))
+	FDMXAttributeName AttributeG_DEPRECATED;
+
+	UPROPERTY(Meta = (DeprecatedProperty, DeprecationMessage = "Deprecated in favor of DMXPixelMappingColorSpace. See ColorSpace member."))
+	FDMXAttributeName AttributeB_DEPRECATED;
+
+	UPROPERTY(Meta = (DeprecatedProperty, DeprecationMessage = "Deprecated in favor of DMXPixelMappingColorSpace. See ColorSpace member."))
+	FDMXAttributeName MonochromeIntensity_DEPRECATED;
+#endif // WITH_EDITORONLY_DATA
 };
diff --git a/Engine/Plugins/VirtualProduction/DMX/DMXPixelMapping/Source/DMXPixelMappingRuntime/Public/Components/DMXPixelMappingMatrixComponent.h b/Engine/Plugins/VirtualProduction/DMX/DMXPixelMapping/Source/DMXPixelMappingRuntime/Public/Components/DMXPixelMappingMatrixComponent.h
--- a/Engine/Plugins/VirtualProduction/DMX/DMXPixelMapping/Source/DMXPixelMappingRuntime/Public/Components/DMXPixelMappingMatrixComponent.h
+++ b/Engine/Plugins/VirtualProduction/DMX/DMXPixelMapping/Source/DMXPixelMappingRuntime/Public/Components/DMXPixelMappingMatrixComponent.h
@@ -7,11 +7,11 @@
 #include "Library/DMXEntityFixtureType.h"
 #include "DMXPixelMappingMatrixComponent.generated.h"
 
+enum class EDMXColorMode : uint8;
 class UDMXLibrary;
+class UDMXPixelMappingColorSpace;
 class UDMXPixelMappingLayoutScript;
 
-class STextBlock;
-enum class EDMXColorMode : uint8;
 
 /**
  * DMX Matrix group component
@@ -100,50 +100,21 @@ private:
 	FVector2D PreEditChangePosition;
 
 public:
+#if WITH_EDITORONLY_DATA
 	UPROPERTY()
 	FDMXEntityFixturePatchRef FixturePatchMatrixRef_DEPRECATED;
+#endif // WITH_EDITORONLY_DATA
 
 	UPROPERTY(VisibleAnywhere, BlueprintReadOnly, Category = "Matrix Settings")
 	FDMXEntityFixturePatchRef FixturePatchRef;
 
-	UPROPERTY(EditAnywhere, BlueprintReadOnly, Category = "Output Settings")
-	EDMXColorMode ColorMode;
-
-	UPROPERTY(EditAnywhere, BlueprintReadOnly, Category = "Output Settings", meta = (DisplayName = "R"))
-	bool AttributeRExpose;
-
-	UPROPERTY(EditAnywhere, BlueprintReadOnly, Category = "Output Settings", meta = (DisplayName = "G"))
-	bool AttributeGExpose;
-
-	UPROPERTY(EditAnywhere, BlueprintReadOnly, Category = "Output Settings", meta = (DisplayName = "B"))
-	bool AttributeBExpose;
-
-	UPROPERTY(EditAnywhere, BlueprintReadOnly, Category = "Output Settings", meta = (DisplayName = "Expose"))
-	bool bMonochromeExpose;
+	/** Sets which color space Pixel Mapping sends */
+	UPROPERTY(Transient, EditAnywhere, NoClear, Category = "Color Space", Meta = (DisplayPriority = 2, DisplayName = "Output Mode", ShowDisplayNames))
+	TSubclassOf<UDMXPixelMappingColorSpace> ColorSpaceClass;
 
-	UPROPERTY(EditAnywhere, BlueprintReadOnly, Category = "Output Settings", meta = (DisplayName = "Invert R"))
-	bool AttributeRInvert;
-
-	UPROPERTY(EditAnywhere, BlueprintReadOnly, Category = "Output Settings", meta = (DisplayName = "Invert G"))
-	bool AttributeGInvert;
-
-	UPROPERTY(EditAnywhere, BlueprintReadOnly, Category = "Output Settings", meta = (DisplayName = "Invert B"))
-	bool AttributeBInvert;
-
-	UPROPERTY(EditAnywhere, BlueprintReadOnly, Category = "Output Settings", meta = (DisplayName = "Invert"))
-	bool bMonochromeInvert;
-
-	UPROPERTY(EditAnywhere, BlueprintReadOnly, Category = "Output Settings", meta = (DisplayName = "R Attribute"))
-	FDMXAttributeName AttributeR;
-
-	UPROPERTY(EditAnywhere, BlueprintReadOnly, Category = "Output Settings", meta = (DisplayName = "G Attribute"))
-	FDMXAttributeName AttributeG;
-
-	UPROPERTY(EditAnywhere, BlueprintReadOnly, Category = "Output Settings", meta = (DisplayName = "B Attribute"))
-	FDMXAttributeName AttributeB;
-
-	UPROPERTY(EditAnywhere, BlueprintReadOnly, Category = "Output Settings", meta = (DisplayName = "Intensity Attribute"))
-	FDMXAttributeName MonochromeIntensity;
+	/** The Color Space currently in use */
+	UPROPERTY(VisibleAnywhere, Instanced, Category = "Color Space")
+	TObjectPtr<UDMXPixelMappingColorSpace> ColorSpace;
 
 	/** Modulators applied to the output before sending DMX */
 	UPROPERTY(EditAnywhere, BlueprintReadOnly, Transient, Category = "Output Settings", meta = (DisplayName = "Output Modulators"))
@@ -165,4 +136,49 @@ public:
 	/** Layout script for the children of this component (hidden in customizations and displayed in its own panel). */
 	UPROPERTY(EditAnywhere, Instanced, Category = "Layout")
 	TObjectPtr<UDMXPixelMappingLayoutScript> LayoutScript;
+
+	//////////////////////
+	// Deprecated Members
+	
+public:
+#if WITH_EDITORONLY_DATA
+	UPROPERTY(Meta = (DeprecatedProperty, DeprecationMessage = "Deprecated in favor of DMXPixelMappingColorSpace. See ColorSpace member."))
+	EDMXColorMode ColorMode_DEPRECATED;
+
+	UPROPERTY(Meta = (DeprecatedProperty, DeprecationMessage = "Deprecated in favor of DMXPixelMappingColorSpace. See ColorSpace member."))
+	bool AttributeRExpose_DEPRECATED;
+
+	UPROPERTY(Meta = (DeprecatedProperty, DeprecationMessage = "Deprecated in favor of DMXPixelMappingColorSpace. See ColorSpace member."))
+	bool AttributeGExpose_DEPRECATED;
+
+	UPROPERTY(Meta = (DeprecatedProperty, DeprecationMessage = "Deprecated in favor of DMXPixelMappingColorSpace. See ColorSpace member."))
+	bool AttributeBExpose_DEPRECATED;
+
+	UPROPERTY(Meta = (DeprecatedProperty, DeprecationMessage = "Deprecated in favor of DMXPixelMappingColorSpace. See ColorSpace member."))
+	bool bMonochromeExpose_DEPRECATED;
+
+	UPROPERTY(Meta = (DeprecatedProperty, DeprecationMessage = "Deprecated in favor of DMXPixelMappingColorSpace. See ColorSpace member."))
+	bool AttributeRInvert_DEPRECATED;
+
+	UPROPERTY(Meta = (DeprecatedProperty, DeprecationMessage = "Deprecated in favor of DMXPixelMappingColorSpace. See ColorSpace member."))
+	bool AttributeGInvert_DEPRECATED;
+
+	UPROPERTY(Meta = (DeprecatedProperty, DeprecationMessage = "Deprecated in favor of DMXPixelMappingColorSpace. See ColorSpace member."))
+	bool AttributeBInvert_DEPRECATED;
+
+	UPROPERTY(Meta = (DeprecatedProperty, DeprecationMessage = "Deprecated in favor of DMXPixelMappingColorSpace. See ColorSpace member."))
+	bool bMonochromeInvert_DEPRECATED;
+
+	UPROPERTY(Meta = (DeprecatedProperty, DeprecationMessage = "Deprecated in favor of DMXPixelMappingColorSpace. See ColorSpace member."))
+	FDMXAttributeName AttributeR_DEPRECATED;
+
+	UPROPERTY(Meta = (DeprecatedProperty, DeprecationMessage = "Deprecated in favor of DMXPixelMappingColorSpace. See ColorSpace member."))
+	FDMXAttributeName AttributeG_DEPRECATED;
+
+	UPROPERTY(Meta = (DeprecatedProperty, DeprecationMessage = "Deprecated in favor of DMXPixelMappingColorSpace. See ColorSpace member."))
+	FDMXAttributeName AttributeB_DEPRECATED;
+
+	UPROPERTY(Meta = (DeprecatedProperty, DeprecationMessage = "Deprecated in favor of DMXPixelMappingColorSpace. See ColorSpace member."))
+	FDMXAttributeName MonochromeIntensity_DEPRECATED;
+#endif // WITH_EDITORONLY_DATA
 };
diff --git a/Engine/Plugins/VirtualProduction/DMX/DMXPixelMapping/Source/DMXPixelMappingRuntime/Public/Components/DMXPixelMappingRendererComponent.h b/Engine/Plugins/VirtualProduction/DMX/DMXPixelMapping/Source/DMXPixelMappingRuntime/Public/Components/DMXPixelMappingRendererComponent.h
--- a/Engine/Plugins/VirtualProduction/DMX/DMXPixelMapping/Source/DMXPixelMappingRuntime/Public/Components/DMXPixelMappingRendererComponent.h
+++ b/Engine/Plugins/VirtualProduction/DMX/DMXPixelMapping/Source/DMXPixelMappingRuntime/Public/Components/DMXPixelMappingRendererComponent.h
@@ -22,10 +22,10 @@ class UTextureRenderTarget2D;
 class UWorld;
 
 
-/**
- * Component for rendering input texture
+/** 
+ * Component for rendering input texture.  
  */
-UCLASS(BlueprintType, Blueprintable)
+UCLASS(BlueprintType)
 class DMXPIXELMAPPINGRUNTIME_API UDMXPixelMappingRendererComponent
 	: public UDMXPixelMappingOutputComponent
 {
@@ -42,7 +42,6 @@ public:
 	virtual void Serialize(FArchive& Ar) override;
 	virtual void PostLoad() override;
 	virtual void PostInitProperties() override;
-
 #if WITH_EDITOR
 	virtual void PostEditChangeChainProperty(FPropertyChangedChainEvent& PropertyChangedChainEvent) override;
 #endif // WITH_EDITOR
@@ -105,10 +104,10 @@ public:
 	 *
 	 * @param InDownsamplePixelParam pixel rendering params
 	 */
-	void AddPixelToDownsampleSet(FDMXPixelMappingDownsamplePixelParam&& InDownsamplePixelParam);
+	void AddPixelToDownsampleSet(FDMXPixelMappingDownsamplePixelParamsV2&& InDownsamplePixelParam);
 
 	/** Get amount of downsample pixels */
-	int32 GetDownsamplePixelNum() const { return DownsamplePixelParams.Num(); }
+	int32 GetDownsamplePixelNum();
 
 	/**
 	 * Pass the downsample CPU buffer from Render Thread to Game Thread and store 
@@ -237,10 +236,13 @@ private:
 	int32 DownsamplePixelCount;
 
 	/** Critical section for set, update and get color array */
-	FCriticalSection DownsampleBufferCS;
+	mutable FCriticalSection DownsampleBufferCS;
 
 	/** Hold the params of the pixels for downsamle rendering */
-	TArray<FDMXPixelMappingDownsamplePixelParam> DownsamplePixelParams;
+	TArray<FDMXPixelMappingDownsamplePixelParamsV2> DownsamplePixelParams;
+
+	/** True once the first render pass completed */
+	bool bWasEverRendered = false;
 
 	/** Initial texture color */
 	static const FLinearColor ClearTextureColor;
diff --git a/Engine/Plugins/VirtualProduction/DMX/DMXPixelMapping/Source/DMXPixelMappingRuntime/Public/Components/DMXPixelMappingRootComponent.h b/Engine/Plugins/VirtualProduction/DMX/DMXPixelMapping/Source/DMXPixelMappingRuntime/Public/Components/DMXPixelMappingRootComponent.h
--- a/Engine/Plugins/VirtualProduction/DMX/DMXPixelMapping/Source/DMXPixelMappingRuntime/Public/Components/DMXPixelMappingRootComponent.h
+++ b/Engine/Plugins/VirtualProduction/DMX/DMXPixelMapping/Source/DMXPixelMappingRuntime/Public/Components/DMXPixelMappingRootComponent.h
@@ -3,8 +3,12 @@
 #pragma once
 
 #include "Components/DMXPixelMappingBaseComponent.h"
+
 #include "DMXPixelMappingRootComponent.generated.h"
 
+class UDMXPixelMappingRendererComponent;
+
+
 /**
  * Root component in the components tree
  */
@@ -16,12 +20,15 @@ class DMXPIXELMAPPINGRUNTIME_API UDMXPixelMappingRootComponent
 
 protected:
 	// ~Begin UObject Interface
+	virtual void PostLoad() override;
 	virtual void BeginDestroy() override;
 	// ~End UObject Interface
 
 public:
 	//~ Begin UDMXPixelMappingBaseComponent implementation
 	virtual const FName& GetNamePrefix() override;
+	virtual void AddChild(UDMXPixelMappingBaseComponent* InComponent) override;
+	virtual void RemoveChild(UDMXPixelMappingBaseComponent* InComponent) override;
 	virtual void ResetDMX() override;
 	virtual void SendDMX() override;
 	virtual void Render() override;
@@ -30,4 +37,15 @@ public:
 
 	/** Check if a Component can be moved under another one (used for copy/move/duplicate) */
 	virtual bool CanBeMovedTo(const UDMXPixelMappingBaseComponent* Component) const override;
+
+	/** Finds a renderer component by its name. Faster than UDMXPixelMappingBaseComponent::FindComponentByName */
+	UDMXPixelMappingRendererComponent* FindRendererComponentByName(const FName& Name) const;
+
+private:
+	/** Called when a component in the PixelMapping was renamed */
+	void OnComponentRenamed(UDMXPixelMapping* PixelMapping , UDMXPixelMappingBaseComponent* RenamedComponent, UObject* OldOuter, const FName OldName);
+
+	/** Chached Renderer Component names and their actual object pointer. Useful to speed up access (see DMXPixelMappingSubsystem) */
+	UPROPERTY(Transient)
+	TMap<FName, TObjectPtr<UDMXPixelMappingRendererComponent>> CachedRendererComponentsByName;
 };
