DMX - Fix CIS Static Analysis Win64 (MSVC) - Dereferencing NULL pointer 'ColorSpace_RGBCMY'

diff --git a/Engine/Plugins/VirtualProduction/DMX/DMXPixelMapping/Source/DMXPixelMappingRuntime/Private/Components/DMXPixelMappingMatrixComponent.cpp b/Engine/Plugins/VirtualProduction/DMX/DMXPixelMapping/Source/DMXPixelMappingRuntime/Private/Components/DMXPixelMappingMatrixComponent.cpp
--- a/Engine/Plugins/VirtualProduction/DMX/DMXPixelMapping/Source/DMXPixelMappingRuntime/Private/Components/DMXPixelMappingMatrixComponent.cpp
+++ b/Engine/Plugins/VirtualProduction/DMX/DMXPixelMapping/Source/DMXPixelMappingRuntime/Private/Components/DMXPixelMappingMatrixComponent.cpp
@@ -47,11 +47,11 @@ void UDMXPixelMappingMatrixComponent::Serialize(FArchive& Ar)
 		if (Ar.CustomVer(FDMXPixelMappingMainStreamObjectVersion::GUID) < FDMXPixelMappingMainStreamObjectVersion::UseDMXPixelMappingColorSpace)
 		{		
 			// Upgrade to the DMXPixelMappingColorSpace default subobject
-			UDMXPixelMappingColorSpace_RGBCMY* ColorSpace_RGBCMY = Cast<UDMXPixelMappingColorSpace_RGBCMY>(ColorSpace);
-			if (!ensureMsgf(ColorSpace_RGBCMY, TEXT("Missing default Subobject ColorSpace")))
+			if (!ensureMsgf(ColorSpace, TEXT("Missing default Subobject ColorSpace")))
 			{
 				ColorSpace = NewObject<UDMXPixelMappingColorSpace_RGBCMY>(this, "ColorSpace");
 			}
+			UDMXPixelMappingColorSpace_RGBCMY* ColorSpace_RGBCMY = Cast<UDMXPixelMappingColorSpace_RGBCMY>(ColorSpace);
 
 			if (bMonochromeExpose_DEPRECATED)
 			{
