DMX - Upgrade path to the new PixelMapping color management is faulty when the asset uses default attribute or color mode values

diff --git a/Engine/Plugins/VirtualProduction/DMX/DMXPixelMapping/Source/DMXPixelMappingRuntime/Private/Components/DMXPixelMappingFixtureGroupItemComponent.cpp b/Engine/Plugins/VirtualProduction/DMX/DMXPixelMapping/Source/DMXPixelMappingRuntime/Private/Components/DMXPixelMappingFixtureGroupItemComponent.cpp
--- a/Engine/Plugins/VirtualProduction/DMX/DMXPixelMapping/Source/DMXPixelMappingRuntime/Private/Components/DMXPixelMappingFixtureGroupItemComponent.cpp
+++ b/Engine/Plugins/VirtualProduction/DMX/DMXPixelMapping/Source/DMXPixelMappingRuntime/Private/Components/DMXPixelMappingFixtureGroupItemComponent.cpp
@@ -35,6 +35,14 @@ UDMXPixelMappingFixtureGroupItemComponent::UDMXPixelMappingFixtureGroupItemCompo
 
 	SetSize(FVector2D(32.f, 32.f));
 
+#if WITH_EDITORONLY_DATA
+	// Even tho deprecated, default values on deprecated properties need be set so they don't load their type's default value.
+	ColorMode_DEPRECATED = EDMXColorMode::CM_RGB;
+	AttributeR_DEPRECATED.SetFromName("Red");
+	AttributeG_DEPRECATED.SetFromName("Green");
+	AttributeB_DEPRECATED.SetFromName("Blue");
+#endif // WITH_EDITORONLY_DATA
+
 #if WITH_EDITOR
 	ZOrder = 2;
 #endif // WITH_EDITOR
@@ -54,9 +62,12 @@ void UDMXPixelMappingFixtureGroupItemComponent::Serialize(FArchive& Ar)
 			UDMXPixelMappingColorSpace_RGBCMY* ColorSpace_RGBCMY = Cast<UDMXPixelMappingColorSpace_RGBCMY>(ColorSpace);
 			checkf(ColorSpace_RGBCMY, TEXT("Missing default Subobject ColorSpace"));
 
-			if (bMonochromeExpose_DEPRECATED)
+			if (ColorMode_DEPRECATED == EDMXColorMode::CM_Monochrome)
 			{
 				ColorSpace_RGBCMY->LuminanceAttribute = MonochromeIntensity_DEPRECATED;
+				ColorSpace_RGBCMY->RedAttribute = FDMXAttributeName(NAME_None);
+				ColorSpace_RGBCMY->GreenAttribute = FDMXAttributeName(NAME_None);
+				ColorSpace_RGBCMY->BlueAttribute = FDMXAttributeName(NAME_None);
 			}
 			else
 			{
@@ -66,6 +77,8 @@ void UDMXPixelMappingFixtureGroupItemComponent::Serialize(FArchive& Ar)
 				ColorSpace_RGBCMY->RedAttribute = AttributeR_DEPRECATED;
 				ColorSpace_RGBCMY->GreenAttribute = AttributeG_DEPRECATED;
 				ColorSpace_RGBCMY->BlueAttribute = AttributeB_DEPRECATED;
+				ColorSpace_RGBCMY->LuminanceType = EDMXPixelMappingLuminanceType_RGBCMY::None;
+				ColorSpace_RGBCMY->LuminanceAttribute = FDMXAttributeName(NAME_None);
 			}
 		}
 	}
