DMX - Fix PixelMapping live preview when outputting DMX does not work

diff --git a/Engine/Plugins/VirtualProduction/DMX/DMXPixelMapping/Source/DMXPixelMappingEditor/Private/Toolkits/DMXPixelMappingToolkit.cpp b/Engine/Plugins/VirtualProduction/DMX/DMXPixelMapping/Source/DMXPixelMappingEditor/Private/Toolkits/DMXPixelMappingToolkit.cpp
--- a/Engine/Plugins/VirtualProduction/DMX/DMXPixelMapping/Source/DMXPixelMappingEditor/Private/Toolkits/DMXPixelMappingToolkit.cpp
+++ b/Engine/Plugins/VirtualProduction/DMX/DMXPixelMapping/Source/DMXPixelMappingEditor/Private/Toolkits/DMXPixelMappingToolkit.cpp
@@ -144,39 +144,44 @@ void FDMXPixelMappingToolkit::Tick(float DeltaTime)
 		return;
 	}
 
-	// render selected component
-	if (!bIsPlayingDMX && !SelectedComponents.IsEmpty())
+	TSet<UDMXPixelMappingRendererComponent*> RenderedRendererComponents;
+
+	// Render preview if required
+	for (const FDMXPixelMappingComponentReference& SelectedComponentRef : SelectedComponents)
 	{
-		for (const FDMXPixelMappingComponentReference& SelectedComponentRef : SelectedComponents)
+		UDMXPixelMappingBaseComponent* SelectedComponent = SelectedComponentRef.Component.Get();
+		if (IsValid(SelectedComponent))
 		{
-			UDMXPixelMappingBaseComponent* SelectedComponent = SelectedComponentRef.Component.Get();
-			if (IsValid(SelectedComponent))
+			// User select root component
+			if (Cast<UDMXPixelMappingRootComponent>(SelectedComponent))
 			{
-				// User select root component
-				if (Cast<UDMXPixelMappingRootComponent>(SelectedComponent))
-				{
-					break;
-				}
+				break;
+			}
 
-				// Try to get renderer component from selected component
-				UDMXPixelMappingRendererComponent* RendererComponent = SelectedComponent->GetRendererComponent();
-				if (!ensureMsgf(RendererComponent, TEXT("Component %s resides in pixelmapping but has no valid renderer."), *SelectedComponent->GetUserFriendlyName()))
-				{
-					break;
-				}
+			// Try to get renderer component from selected component
+			UDMXPixelMappingRendererComponent* RendererComponent = SelectedComponent->GetRendererComponent();
+			if (!ensureMsgf(RendererComponent, TEXT("Component %s resides in pixelmapping but has no valid renderer."), *SelectedComponent->GetUserFriendlyName()))
+			{
+				break;
+			}
 
-				// Render
+			// Render preview depends on the render component being 'rendered'
+			if (!RenderedRendererComponents.Contains(RendererComponent))
+			{
 				RendererComponent->Render();
+				RenderedRendererComponents.Add(RendererComponent);
+			}
 
-				// Render preview
-				RendererComponent->RenderEditorPreviewTexture();
+			// Render preview
+			RendererComponent->RenderEditorPreviewTexture();
 
-				// Render only once for all selected components
-				break;
-			}
+
+			// Render only once for all selected components
+			break;
 		}
 	}
 
+	// Render Output and Send DMX if required
 	if (bIsPlayingDMX && !bRequestStopSendingDMX)
 	{
 		if (bTogglePlayDMXAll) // Send to all
@@ -188,7 +193,6 @@ void FDMXPixelMappingToolkit::Tick(float DeltaTime)
 		{
 			for (FDMXPixelMappingComponentReference& SelectedComponentRef : SelectedComponents)
 			{
-				TSet<UDMXPixelMappingRendererComponent*> RenderedRendererComponents;
 				if (UDMXPixelMappingBaseComponent* SelectedComponent = SelectedComponentRef.Component.Get())
 				{
 					// Try to get renderer component from selected component
