DMX - Fix Pixel Mapping Output Colorspace doesn't affect DMX Output, fix Luminance Min and Max are not applied

diff --git a/Engine/Plugins/VirtualProduction/DMX/DMXPixelMapping/Source/DMXPixelMappingRuntime/Private/ColorSpace/DMXPixelMappingColorSpace_RGBCMY.cpp b/Engine/Plugins/VirtualProduction/DMX/DMXPixelMapping/Source/DMXPixelMappingRuntime/Private/ColorSpace/DMXPixelMappingColorSpace_RGBCMY.cpp
--- a/Engine/Plugins/VirtualProduction/DMX/DMXPixelMapping/Source/DMXPixelMappingRuntime/Private/ColorSpace/DMXPixelMappingColorSpace_RGBCMY.cpp
+++ b/Engine/Plugins/VirtualProduction/DMX/DMXPixelMapping/Source/DMXPixelMappingRuntime/Private/ColorSpace/DMXPixelMappingColorSpace_RGBCMY.cpp
@@ -43,15 +43,15 @@ void UDMXPixelMappingColorSpace_RGBCMY::SetRGBA(const FLinearColor& InColor)
 		if (LuminanceType == EDMXPixelMappingLuminanceType_RGBCMY::FromColor)
 		{
 			const float LuminanceFromColor = OutputColorSpace.GetLuminance(CalibratedColor);
-			SetAttributeValue(LuminanceAttribute, LuminanceFromColor);
+			SetAttributeValue(LuminanceAttribute, FMath::Clamp(LuminanceFromColor, MinLuminance, MaxLuminance));
 		}
 		else if (LuminanceType == EDMXPixelMappingLuminanceType_RGBCMY::Constant)
 		{
-			SetAttributeValue(LuminanceAttribute, Luminance);
+			SetAttributeValue(LuminanceAttribute, FMath::Clamp(Luminance, MinLuminance, MaxLuminance));
 		}
 		else if (LuminanceType == EDMXPixelMappingLuminanceType_RGBCMY::FromAlpha)
 		{
-			SetAttributeValue(LuminanceAttribute, InColor.A);
+			SetAttributeValue(LuminanceAttribute, FMath::Clamp(InColor.A, MinLuminance, MaxLuminance));
 		}
 	}
 }
@@ -74,11 +74,27 @@ void UDMXPixelMappingColorSpace_RGBCMY::PostEditChangeChainProperty(FPropertyCha
 	{
 		ClearCachedAttributeValues();
 	}
-	else if (PropertyName == GET_MEMBER_NAME_CHECKED(UDMXPixelMappingColorSpace_RGBCMY, OutputColorSpace) ||
+	else if (PropertyName == GET_MEMBER_NAME_CHECKED(UDMXPixelMappingColorSpace_RGBCMY, PixelMappingOutputColorSpace) ||
 		PropertyName == GET_MEMBER_NAME_CHECKED(UDMXPixelMappingColorSpace, bUseWorkingColorSpaceForInput))
 	{
 		UpdateColorSpaceAndTransform();
 	}
+	else if (PropertyName == GET_MEMBER_NAME_CHECKED(UDMXPixelMappingColorSpace_RGBCMY, MinLuminance))
+	{
+		if (MaxLuminance < MinLuminance)
+		{
+			Modify();
+			MaxLuminance = MinLuminance;
+		}
+	}
+	else if (PropertyName == GET_MEMBER_NAME_CHECKED(UDMXPixelMappingColorSpace_RGBCMY, MaxLuminance))
+	{
+		if (MinLuminance > MaxLuminance)
+		{
+			Modify();
+			MinLuminance = MaxLuminance;
+		}
+	}
 }
 #endif // WITH_EDITOR
 
diff --git a/Engine/Plugins/VirtualProduction/DMX/DMXPixelMapping/Source/DMXPixelMappingRuntime/Private/ColorSpace/DMXPixelMappingColorSpace_XYZ.cpp b/Engine/Plugins/VirtualProduction/DMX/DMXPixelMapping/Source/DMXPixelMappingRuntime/Private/ColorSpace/DMXPixelMappingColorSpace_XYZ.cpp
--- a/Engine/Plugins/VirtualProduction/DMX/DMXPixelMapping/Source/DMXPixelMappingRuntime/Private/ColorSpace/DMXPixelMappingColorSpace_XYZ.cpp
+++ b/Engine/Plugins/VirtualProduction/DMX/DMXPixelMapping/Source/DMXPixelMappingRuntime/Private/ColorSpace/DMXPixelMappingColorSpace_XYZ.cpp
@@ -26,7 +26,8 @@ void UDMXPixelMappingColorSpace_XYZ::SetRGBA(const FLinearColor& InColor)
 
 	if (LuminanceAttribute.IsValid())
 	{
-		SetAttributeValue(LuminanceAttribute, XYZW.Y);
+		const float Luminance = FMath::Clamp(XYZW.Y, MinLuminance, MaxLuminance);
+		SetAttributeValue(LuminanceAttribute, Luminance);
 	}
 
 	if (ZAttribute.IsValid())
@@ -48,5 +49,21 @@ void UDMXPixelMappingColorSpace_XYZ::PostEditChangeProperty(FPropertyChangedEven
 	{
 		ClearCachedAttributeValues();
 	}
+	else if (PropertyName == GET_MEMBER_NAME_CHECKED(UDMXPixelMappingColorSpace_XYZ, MinLuminance))
+	{
+		if (MaxLuminance < MinLuminance)
+		{
+			Modify();
+			MaxLuminance = MinLuminance;
+		}
+	}
+	else if (PropertyName == GET_MEMBER_NAME_CHECKED(UDMXPixelMappingColorSpace_XYZ, MaxLuminance))
+	{
+		if (MinLuminance > MaxLuminance)
+		{
+			Modify();
+			MinLuminance = MaxLuminance;
+		}
+	}
 }
 #endif // WITH_EDITOR
diff --git a/Engine/Plugins/VirtualProduction/DMX/DMXPixelMapping/Source/DMXPixelMappingRuntime/Private/ColorSpace/DMXPixelMappingColorSpace_XYZ.h b/Engine/Plugins/VirtualProduction/DMX/DMXPixelMapping/Source/DMXPixelMappingRuntime/Private/ColorSpace/DMXPixelMappingColorSpace_XYZ.h
--- a/Engine/Plugins/VirtualProduction/DMX/DMXPixelMapping/Source/DMXPixelMappingRuntime/Private/ColorSpace/DMXPixelMappingColorSpace_XYZ.h
+++ b/Engine/Plugins/VirtualProduction/DMX/DMXPixelMapping/Source/DMXPixelMappingRuntime/Private/ColorSpace/DMXPixelMappingColorSpace_XYZ.h
@@ -37,10 +37,6 @@ public:
 	UPROPERTY(EditAnywhere, Category = "XYZ")
 	FDMXAttributeName LuminanceAttribute;
 
-	/** Scale Factor for Luminance, e.g. when 0.5, a Luminance of 1 will be output as 0.5. */
-	UPROPERTY(EditAnywhere, Category = "Luminance", Meta = (ClampMin = 0.0, ClampMax = 1.0, UIMin = 0.0, UIMax = 1.0, EditConditionHides, EditCondition = "LuminanceType == EDMXPixelMappingLuminanceType_RGBCMY::Constant"))
-	float LuminanceScale = 1.f;
-
 	/** Min Luminance */
 	UPROPERTY(EditAnywhere, Category = "Luminance", Meta = (ClampMin = 0.0, ClampMax = 1.0, UIMin = 0.0, UIMax = 1.0, EditConditionHides, EditCondition = "LuminanceType == EDMXPixelMappingLuminanceType_RGBCMY::FromWhite || LuminanceType == EDMXPixelMappingLuminanceType_RGBCMY::FromAlpha"))
 	float MinLuminance = 0.f;
diff --git a/Engine/Plugins/VirtualProduction/DMX/DMXPixelMapping/Source/DMXPixelMappingRuntime/Private/ColorSpace/DMXPixelMappingColorSpace_xyY.cpp b/Engine/Plugins/VirtualProduction/DMX/DMXPixelMapping/Source/DMXPixelMappingRuntime/Private/ColorSpace/DMXPixelMappingColorSpace_xyY.cpp
--- a/Engine/Plugins/VirtualProduction/DMX/DMXPixelMapping/Source/DMXPixelMappingRuntime/Private/ColorSpace/DMXPixelMappingColorSpace_xyY.cpp
+++ b/Engine/Plugins/VirtualProduction/DMX/DMXPixelMapping/Source/DMXPixelMappingRuntime/Private/ColorSpace/DMXPixelMappingColorSpace_xyY.cpp
@@ -38,7 +38,7 @@ void UDMXPixelMappingColorSpace_xyY::SetRGBA(const FLinearColor& InColor)
 
 	if (LuminanceAttribute.IsValid())
 	{
-		const float Luminance = FMath::Clamp(XYZW.Y * LuminanceScale, MinLuminance, MaxLuminance);
+		const float Luminance = FMath::Clamp(XYZW.Y, MinLuminance, MaxLuminance);
 		SetAttributeValue(LuminanceAttribute, Luminance);
 	}
 }
@@ -56,5 +56,21 @@ void UDMXPixelMappingColorSpace_xyY::PostEditChangeProperty(FPropertyChangedEven
 	{
 		ClearCachedAttributeValues();
 	}
+	else if (PropertyName == GET_MEMBER_NAME_CHECKED(UDMXPixelMappingColorSpace_xyY, MinLuminance))
+	{
+		if (MaxLuminance < MinLuminance)
+		{
+			Modify();
+			MaxLuminance = MinLuminance;
+		}
+	}
+	else if (PropertyName == GET_MEMBER_NAME_CHECKED(UDMXPixelMappingColorSpace_xyY, MaxLuminance))
+	{
+		if (MinLuminance > MaxLuminance)
+		{
+			Modify();
+			MinLuminance = MaxLuminance;
+		}
+	}
 }
 #endif // WITH_EDITOR
diff --git a/Engine/Plugins/VirtualProduction/DMX/DMXPixelMapping/Source/DMXPixelMappingRuntime/Private/ColorSpace/DMXPixelMappingColorSpace_xyY.h b/Engine/Plugins/VirtualProduction/DMX/DMXPixelMapping/Source/DMXPixelMappingRuntime/Private/ColorSpace/DMXPixelMappingColorSpace_xyY.h
--- a/Engine/Plugins/VirtualProduction/DMX/DMXPixelMapping/Source/DMXPixelMappingRuntime/Private/ColorSpace/DMXPixelMappingColorSpace_xyY.h
+++ b/Engine/Plugins/VirtualProduction/DMX/DMXPixelMapping/Source/DMXPixelMappingRuntime/Private/ColorSpace/DMXPixelMappingColorSpace_xyY.h
@@ -36,10 +36,6 @@ public:
 	UPROPERTY(EditAnywhere, Category = "Luminance")
 	FDMXAttributeName LuminanceAttribute;
 
-	/** Scale Factor for Luminance, e.g. when 0.5, a Luminance of 1 will be output as 0.5. */
-	UPROPERTY(EditAnywhere, Category = "Luminance", Meta = (ClampMin = 0.0, ClampMax = 1.0, SliderMin = 0.0, SliderMax = 1.0, EditConditionHides, EditCondition = "LuminanceType == EDMXPixelMappingLuminanceType_RGBCMY::Constant"))
-	float LuminanceScale = 1.f;
-
 	/** Min Luminance */
 	UPROPERTY(EditAnywhere, Category = "Luminance", Meta = (ClampMin = 0.0, ClampMax = 1.0, SliderMin = 0.0, SliderMax = 1.0, EditConditionHides, EditCondition = "LuminanceType == EDMXPixelMappingLuminanceType_RGBCMY::FromWhite || LuminanceType == EDMXPixelMappingLuminanceType_RGBCMY::FromAlpha"))
 	float MinLuminance = 0.f;
diff --git a/Engine/Plugins/VirtualProduction/DMX/DMXPixelMapping/Source/DMXPixelMappingRuntime/Private/Components/DMXPixelMappingFixtureGroupItemComponent.cpp b/Engine/Plugins/VirtualProduction/DMX/DMXPixelMapping/Source/DMXPixelMappingRuntime/Private/Components/DMXPixelMappingFixtureGroupItemComponent.cpp
--- a/Engine/Plugins/VirtualProduction/DMX/DMXPixelMapping/Source/DMXPixelMappingRuntime/Private/Components/DMXPixelMappingFixtureGroupItemComponent.cpp
+++ b/Engine/Plugins/VirtualProduction/DMX/DMXPixelMapping/Source/DMXPixelMappingRuntime/Private/Components/DMXPixelMappingFixtureGroupItemComponent.cpp
@@ -31,6 +31,7 @@ DECLARE_CYCLE_STAT(TEXT("Send Fixture Group Item"), STAT_DMXPixelMaping_FixtureG
 UDMXPixelMappingFixtureGroupItemComponent::UDMXPixelMappingFixtureGroupItemComponent()
 	: DownsamplePixelIndex(0)
 {
+	ColorSpaceClass = UDMXPixelMappingColorSpace_RGBCMY::StaticClass();
 	ColorSpace = CreateDefaultSubobject<UDMXPixelMappingColorSpace_RGBCMY>("ColorSpace");
 
 	SetSize(FVector2D(32.f, 32.f));
@@ -56,11 +57,14 @@ void UDMXPixelMappingFixtureGroupItemComponent::Serialize(FArchive& Ar)
 	Ar.UsingCustomVersion(FDMXPixelMappingMainStreamObjectVersion::GUID);
 	if (Ar.IsLoading())
 	{
-		// Upgrade to the DMXPixelMappingColorSpace default subobject
 		if (Ar.CustomVer(FDMXPixelMappingMainStreamObjectVersion::GUID) < FDMXPixelMappingMainStreamObjectVersion::UseDMXPixelMappingColorSpace)
 		{
+			// Upgrade to the DMXPixelMappingColorSpace default subobject
+			if (!ensureMsgf(ColorSpace, TEXT("Missing default Subobject ColorSpace")))
+			{
+				ColorSpace = NewObject<UDMXPixelMappingColorSpace_RGBCMY>(this, "ColorSpace");
+			}
 			UDMXPixelMappingColorSpace_RGBCMY* ColorSpace_RGBCMY = Cast<UDMXPixelMappingColorSpace_RGBCMY>(ColorSpace);
-			checkf(ColorSpace_RGBCMY, TEXT("Missing default Subobject ColorSpace"));
 
 			if (ColorMode_DEPRECATED == EDMXColorMode::CM_Monochrome)
 			{
diff --git a/Engine/Plugins/VirtualProduction/DMX/DMXPixelMapping/Source/DMXPixelMappingRuntime/Private/Components/DMXPixelMappingMatrixComponent.cpp b/Engine/Plugins/VirtualProduction/DMX/DMXPixelMapping/Source/DMXPixelMappingRuntime/Private/Components/DMXPixelMappingMatrixComponent.cpp
--- a/Engine/Plugins/VirtualProduction/DMX/DMXPixelMapping/Source/DMXPixelMappingRuntime/Private/Components/DMXPixelMappingMatrixComponent.cpp
+++ b/Engine/Plugins/VirtualProduction/DMX/DMXPixelMapping/Source/DMXPixelMappingRuntime/Private/Components/DMXPixelMappingMatrixComponent.cpp
@@ -25,9 +25,18 @@
 
 UDMXPixelMappingMatrixComponent::UDMXPixelMappingMatrixComponent()
 {
+	ColorSpaceClass = UDMXPixelMappingColorSpace_RGBCMY::StaticClass();
 	ColorSpace = CreateDefaultSubobject<UDMXPixelMappingColorSpace_RGBCMY>("ColorSpace");
 
 	SetSize(FVector2D(32.f, 32.f));
+
+#if WITH_EDITORONLY_DATA
+	// Even tho deprecated, default values on deprecated properties need be set so they don't load their type's default value.
+	ColorMode_DEPRECATED = EDMXColorMode::CM_RGB;
+	AttributeR_DEPRECATED.SetFromName("Red");
+	AttributeG_DEPRECATED.SetFromName("Green");
+	AttributeB_DEPRECATED.SetFromName("Blue");
+#endif // WITH_EDITORONLY_DATA
 }
 
 void UDMXPixelMappingMatrixComponent::Serialize(FArchive& Ar)
@@ -53,9 +62,12 @@ void UDMXPixelMappingMatrixComponent::Serialize(FArchive& Ar)
 			}
 			UDMXPixelMappingColorSpace_RGBCMY* ColorSpace_RGBCMY = Cast<UDMXPixelMappingColorSpace_RGBCMY>(ColorSpace);
 
-			if (bMonochromeExpose_DEPRECATED)
+			if (ColorMode_DEPRECATED == EDMXColorMode::CM_Monochrome)
 			{
 				ColorSpace_RGBCMY->LuminanceAttribute = MonochromeIntensity_DEPRECATED;
+				ColorSpace_RGBCMY->RedAttribute = FDMXAttributeName(NAME_None);
+				ColorSpace_RGBCMY->GreenAttribute = FDMXAttributeName(NAME_None);
+				ColorSpace_RGBCMY->BlueAttribute = FDMXAttributeName(NAME_None);
 			}
 			else
 			{
@@ -65,6 +77,8 @@ void UDMXPixelMappingMatrixComponent::Serialize(FArchive& Ar)
 				ColorSpace_RGBCMY->RedAttribute = AttributeR_DEPRECATED;
 				ColorSpace_RGBCMY->GreenAttribute = AttributeG_DEPRECATED;
 				ColorSpace_RGBCMY->BlueAttribute = AttributeB_DEPRECATED;
+				ColorSpace_RGBCMY->LuminanceType = EDMXPixelMappingLuminanceType_RGBCMY::None;
+				ColorSpace_RGBCMY->LuminanceAttribute = FDMXAttributeName(NAME_None);
 			}
 		}
 	}
diff --git a/Engine/Plugins/VirtualProduction/DMX/DMXPixelMapping/Source/DMXPixelMappingRuntime/Public/Components/DMXPixelMappingFixtureGroupItemComponent.h b/Engine/Plugins/VirtualProduction/DMX/DMXPixelMapping/Source/DMXPixelMappingRuntime/Public/Components/DMXPixelMappingFixtureGroupItemComponent.h
--- a/Engine/Plugins/VirtualProduction/DMX/DMXPixelMapping/Source/DMXPixelMappingRuntime/Public/Components/DMXPixelMappingFixtureGroupItemComponent.h
+++ b/Engine/Plugins/VirtualProduction/DMX/DMXPixelMapping/Source/DMXPixelMappingRuntime/Public/Components/DMXPixelMappingFixtureGroupItemComponent.h
@@ -85,7 +85,7 @@ public:
 	TObjectPtr<UDMXPixelMappingColorSpace> ColorSpace;
 
 	/** Modulators applied to the output before sending DMX */
-	UPROPERTY(Transient, EditAnywhere, BlueprintReadOnly, Transient, Category = "Output Settings", meta = (DisplayName = "Output Modulators"))
+	UPROPERTY(Transient, EditAnywhere, BlueprintReadOnly, Category = "Output Settings", Meta = (DisplayName = "Output Modulators"))
 	TArray<TSubclassOf<UDMXModulator>> ModulatorClasses;
 
 	/** Modulator instances applied to this component */
