DMX - PixelMapping now correctly layouts patches when 'add all patches' is clicked

diff --git a/Engine/Plugins/VirtualProduction/DMX/DMXPixelMapping/Source/DMXPixelMappingEditor/Private/Customizations/DMXPixelMappingDetailCustomization_FixtureGroup.cpp b/Engine/Plugins/VirtualProduction/DMX/DMXPixelMapping/Source/DMXPixelMappingEditor/Private/Customizations/DMXPixelMappingDetailCustomization_FixtureGroup.cpp
--- a/Engine/Plugins/VirtualProduction/DMX/DMXPixelMapping/Source/DMXPixelMappingEditor/Private/Customizations/DMXPixelMappingDetailCustomization_FixtureGroup.cpp
+++ b/Engine/Plugins/VirtualProduction/DMX/DMXPixelMapping/Source/DMXPixelMappingEditor/Private/Customizations/DMXPixelMappingDetailCustomization_FixtureGroup.cpp
@@ -208,27 +208,24 @@ FReply FDMXPixelMappingDetailCustomization_FixtureGroup::OnAddAllPatchesClicked(
 	}
 
 	// Layout new components inside the group
-	const int32 Columns = FMath::Sqrt((float)NewComponents.Num());
+	const int32 Columns = FMath::RoundFromZero(FMath::Sqrt((float)NewComponents.Num()));
 	const int32 Rows = FMath::RoundFromZero((float)NewComponents.Num() / Columns);
-	const float Size = FMath::Min(GroupComponent->GetSize().X / Columns, GroupComponent->GetSize().Y / Rows);
-	if (Size < 1.f)
-	{
-		GroupComponent->SetSize(FVector2D(Size * Columns, Size * Rows));
-	}
-	
+	const FVector2D Size = FVector2D(GroupComponent->GetSize().X / Columns, GroupComponent->GetSize().Y / Rows);
+		
 	const FVector2D ParentPosition = GroupComponent->GetPosition();
 	int32 Column = -1;
 	int32 Row = 0;
 	for (UDMXPixelMappingOutputComponent* NewComponent : NewComponents)
 	{
 		Column++;
-		if (Column > Columns)
+		if (Column >= Columns)
 		{
 			Column = 0;
 			Row++;
 		}
-		const FVector2D Position = ParentPosition + FVector2D(Column * Size, Row * Size);
+		const FVector2D Position = ParentPosition + FVector2D(Column * Size.X, Row * Size.Y);
 		NewComponent->SetPosition(Position);
+		NewComponent->SetSize(Size);
 	}
 
 	return FReply::Handled();
