DMX - Fix PixelMapping Output Components show their object name instead of the user set name in editor

diff --git a/Engine/Plugins/VirtualProduction/DMX/DMXPixelMapping/Source/DMXPixelMappingEditor/Private/ViewModels/DMXPixelMappingOutputComponentModel.cpp b/Engine/Plugins/VirtualProduction/DMX/DMXPixelMapping/Source/DMXPixelMappingEditor/Private/ViewModels/DMXPixelMappingOutputComponentModel.cpp
--- a/Engine/Plugins/VirtualProduction/DMX/DMXPixelMapping/Source/DMXPixelMappingEditor/Private/ViewModels/DMXPixelMappingOutputComponentModel.cpp
+++ b/Engine/Plugins/VirtualProduction/DMX/DMXPixelMapping/Source/DMXPixelMappingEditor/Private/ViewModels/DMXPixelMappingOutputComponentModel.cpp
@@ -64,7 +64,7 @@ FText FDMXPixelMappingOutputComponentModel::GetName() const
 {
 	if (UDMXPixelMappingOutputComponent* OutputComponent = Cast<UDMXPixelMappingOutputComponent>(WeakOutputComponent.Get()))
 	{
-		return FText::FromString(OutputComponent->GetName());
+		return FText::FromString(OutputComponent->GetUserFriendlyName());
 	}
 
 	return FText::GetEmpty();
diff --git a/Engine/Plugins/VirtualProduction/DMX/DMXPixelMapping/Source/DMXPixelMappingEditor/Private/Views/SDMXPixelMappingPreviewView.cpp b/Engine/Plugins/VirtualProduction/DMX/DMXPixelMapping/Source/DMXPixelMappingEditor/Private/Views/SDMXPixelMappingPreviewView.cpp
--- a/Engine/Plugins/VirtualProduction/DMX/DMXPixelMapping/Source/DMXPixelMappingEditor/Private/Views/SDMXPixelMappingPreviewView.cpp
+++ b/Engine/Plugins/VirtualProduction/DMX/DMXPixelMapping/Source/DMXPixelMappingEditor/Private/Views/SDMXPixelMappingPreviewView.cpp
@@ -384,7 +384,7 @@ FText SDMXPixelMappingPreviewView::GetSelectedComponentNameText() const
 	UDMXPixelMappingBaseComponent* BaseComponent = SelectedComponentRef.GetComponent();
 	if (BaseComponent)
 	{
-		return FText::FromString(BaseComponent->GetName());
+		return FText::FromString(BaseComponent->GetUserFriendlyName());
 	}
 
 	return FText();
