DMX - PixelMapping brightness no correctly affects the output

diff --git a/Engine/Plugins/VirtualProduction/DMX/DMXPixelMapping/Shaders/Private/DMXPixelMapping.usf b/Engine/Plugins/VirtualProduction/DMX/DMXPixelMapping/Shaders/Private/DMXPixelMapping.usf
--- a/Engine/Plugins/VirtualProduction/DMX/DMXPixelMapping/Shaders/Private/DMXPixelMapping.usf
+++ b/Engine/Plugins/VirtualProduction/DMX/DMXPixelMapping/Shaders/Private/DMXPixelMapping.usf
@@ -83,8 +83,9 @@ OutputPS DMXPixelMappingPS(OutputVS IN)
 		}
 		
 		OUT.Color /= (NUM_CELLS + 1);
-		
 #endif
+
+		OUT.Color *= Brightness;
 	}
 			
 	return OUT;
diff --git a/Engine/Plugins/VirtualProduction/DMX/DMXPixelMapping/Shaders/Private/DMXPixelMapping.ush b/Engine/Plugins/VirtualProduction/DMX/DMXPixelMapping/Shaders/Private/DMXPixelMapping.ush
--- a/Engine/Plugins/VirtualProduction/DMX/DMXPixelMapping/Shaders/Private/DMXPixelMapping.ush
+++ b/Engine/Plugins/VirtualProduction/DMX/DMXPixelMapping/Shaders/Private/DMXPixelMapping.ush
@@ -17,6 +17,7 @@ Texture2D InputTexture;
 SamplerState InputSampler;
 
 // PS Parameters
+float Brightness;
 int2 InputTextureSize;
 int2 OutputTextureSize;
 float2 UVCellSize;
diff --git a/Engine/Plugins/VirtualProduction/DMX/DMXPixelMapping/Source/DMXPixelMappingRenderer/Private/DMXPixelMappingRenderer.cpp b/Engine/Plugins/VirtualProduction/DMX/DMXPixelMapping/Source/DMXPixelMappingRenderer/Private/DMXPixelMappingRenderer.cpp
--- a/Engine/Plugins/VirtualProduction/DMX/DMXPixelMapping/Source/DMXPixelMappingRenderer/Private/DMXPixelMappingRenderer.cpp
+++ b/Engine/Plugins/VirtualProduction/DMX/DMXPixelMapping/Source/DMXPixelMappingRenderer/Private/DMXPixelMappingRenderer.cpp
@@ -74,6 +74,7 @@ public:
 		SHADER_PARAMETER_TEXTURE(Texture2D, InputTexture)
 		SHADER_PARAMETER_SAMPLER(SamplerState, InputSampler)
 
+		SHADER_PARAMETER(float, Brightness)
 		SHADER_PARAMETER(FIntPoint, InputTextureSize)
 		SHADER_PARAMETER(FIntPoint, OutputTextureSize)
 		SHADER_PARAMETER(FVector2f, UVCellSize)
@@ -193,6 +194,7 @@ void FDMXPixelMappingRenderer::DownsampleRender(
 					// Set pixel shader buffer
 					FDMXPixelMappingRendererPS::FParameters PSParameters;
 					PSParameters.InputTexture = InputTextureRHI;
+					PSParameters.Brightness = Brightness;
 					PSParameters.InputSampler = TStaticSamplerState<SF_Trilinear, AM_Clamp, AM_Clamp, AM_Clamp>::GetRHI();
 					PSParameters.InputTextureSize = InputTextureSize;
 					PSParameters.OutputTextureSize = OutputTextureSize;
diff --git a/Engine/Plugins/VirtualProduction/DMX/DMXPixelMapping/Source/DMXPixelMappingRuntime/Private/Components/DMXPixelMappingRendererComponent.cpp b/Engine/Plugins/VirtualProduction/DMX/DMXPixelMapping/Source/DMXPixelMappingRuntime/Private/Components/DMXPixelMappingRendererComponent.cpp
--- a/Engine/Plugins/VirtualProduction/DMX/DMXPixelMapping/Source/DMXPixelMappingRuntime/Private/Components/DMXPixelMappingRendererComponent.cpp
+++ b/Engine/Plugins/VirtualProduction/DMX/DMXPixelMapping/Source/DMXPixelMappingRuntime/Private/Components/DMXPixelMappingRendererComponent.cpp
@@ -108,20 +108,13 @@ void UDMXPixelMappingRendererComponent::Serialize(FArchive& Ar)
 	}
 }
 
-void UDMXPixelMappingRendererComponent::PostLoad()
-{
-	Super::PostLoad();
-	Initialize();
-}
-
 void UDMXPixelMappingRendererComponent::PostInitProperties()
 {
 	Super::PostInitProperties();
-
-	const TSharedPtr<IDMXPixelMappingRenderer>& Renderer = GetRenderer();
-	if (Renderer.IsValid())
+	
+	if (!IsTemplate())
 	{
-		Renderer->SetBrightness(Brightness);
+		Initialize();
 	}
 }
 
@@ -492,6 +485,7 @@ void UDMXPixelMappingRendererComponent::Initialize()
 	if (!PixelMappingRenderer.IsValid())
 	{
 		PixelMappingRenderer = IDMXPixelMappingRendererModule::Get().CreateRenderer();
+		PixelMappingRenderer->SetBrightness(Brightness);
 	}
 
 #if WITH_EDITOR
diff --git a/Engine/Plugins/VirtualProduction/DMX/DMXPixelMapping/Source/DMXPixelMappingRuntime/Public/Components/DMXPixelMappingRendererComponent.h b/Engine/Plugins/VirtualProduction/DMX/DMXPixelMapping/Source/DMXPixelMappingRuntime/Public/Components/DMXPixelMappingRendererComponent.h
--- a/Engine/Plugins/VirtualProduction/DMX/DMXPixelMapping/Source/DMXPixelMappingRuntime/Public/Components/DMXPixelMappingRendererComponent.h
+++ b/Engine/Plugins/VirtualProduction/DMX/DMXPixelMapping/Source/DMXPixelMappingRuntime/Public/Components/DMXPixelMappingRendererComponent.h
@@ -40,7 +40,6 @@ public:
 
 	//~ Begin UObject implementation
 	virtual void Serialize(FArchive& Ar) override;
-	virtual void PostLoad() override;
 	virtual void PostInitProperties() override;
 #if WITH_EDITOR
 	virtual void PostEditChangeChainProperty(FPropertyChangedChainEvent& PropertyChangedChainEvent) override;
