corrected a delegate type. right now it's not a problem but we are going to change its types

diff --git a/Engine/Plugins/Experimental/CommonUI/Source/CommonUI/Private/Widgets/CommonActivatableWidgetContainer.cpp b/Engine/Plugins/Experimental/CommonUI/Source/CommonUI/Private/Widgets/CommonActivatableWidgetContainer.cpp
--- a/Engine/Plugins/Experimental/CommonUI/Source/CommonUI/Private/Widgets/CommonActivatableWidgetContainer.cpp
+++ b/Engine/Plugins/Experimental/CommonUI/Source/CommonUI/Private/Widgets/CommonActivatableWidgetContainer.cpp
@@ -264,7 +264,7 @@ void UCommonActivatableWidgetContainerBase::HandleActiveIndexChanged(int32 Activ
 		if (UWorld* MyWorld = GetWorld())
 		{
 			FTimerManager& TimerManager = MyWorld->GetTimerManager();
-			TimerManager.SetTimerForNextTick(FSimpleDelegate::CreateWeakLambda(this, [this]() { InvalidateLayoutAndVolatility(); }));
+			TimerManager.SetTimerForNextTick(FTimerDelegate::CreateWeakLambda(this, [this]() { InvalidateLayoutAndVolatility(); }));
 		}
 	}
 	else
diff --git a/Engine/Plugins/VirtualProduction/DMX/DMXPixelMapping/Source/DMXPixelMappingEditor/Private/Customizations/DMXPixelMappingDetailCustomization_FixtureGroup.cpp b/Engine/Plugins/VirtualProduction/DMX/DMXPixelMapping/Source/DMXPixelMappingEditor/Private/Customizations/DMXPixelMappingDetailCustomization_FixtureGroup.cpp
--- a/Engine/Plugins/VirtualProduction/DMX/DMXPixelMapping/Source/DMXPixelMappingEditor/Private/Customizations/DMXPixelMappingDetailCustomization_FixtureGroup.cpp
+++ b/Engine/Plugins/VirtualProduction/DMX/DMXPixelMapping/Source/DMXPixelMappingEditor/Private/Customizations/DMXPixelMappingDetailCustomization_FixtureGroup.cpp
@@ -318,7 +318,7 @@ void FDMXPixelMappingDetailCustomization_FixtureGroup::OnSizePropertyChanged(con
 {
 	if(PropertyChangedEvent.ChangeType == EPropertyChangeType::Interactive)
 	{
-		GEditor->GetTimerManager()->SetTimerForNextTick(FSimpleDelegate::CreateSP(this, &FDMXPixelMappingDetailCustomization_FixtureGroup::HandleSizePropertyChanged));
+		GEditor->GetTimerManager()->SetTimerForNextTick(FTimerDelegate::CreateSP(this, &FDMXPixelMappingDetailCustomization_FixtureGroup::HandleSizePropertyChanged));
 	}
 	else
 	{
