DMX - Apply new PixelMapping Design

diff --git a/Engine/Plugins/VirtualProduction/DMX/DMXEngine/Source/DMXEditor/Private/Widgets/FixturePatch/SDMXReadOnlyFixturePatchList.cpp b/Engine/Plugins/VirtualProduction/DMX/DMXEngine/Source/DMXEditor/Private/Widgets/FixturePatch/SDMXReadOnlyFixturePatchList.cpp
--- a/Engine/Plugins/VirtualProduction/DMX/DMXEngine/Source/DMXEditor/Private/Widgets/FixturePatch/SDMXReadOnlyFixturePatchList.cpp
+++ b/Engine/Plugins/VirtualProduction/DMX/DMXEngine/Source/DMXEditor/Private/Widgets/FixturePatch/SDMXReadOnlyFixturePatchList.cpp
@@ -10,6 +10,7 @@
 
 #include "Editor.h"
 #include "TimerManager.h"
+#include "Algo/Find.h"
 #include "Algo/Sort.h"
 #include "Framework/Application/SlateApplication.h"
 #include "Framework/MultiBox/MultiBoxBuilder.h"
@@ -21,6 +22,46 @@
 
 #define LOCTEXT_NAMESPACE "SDMXReadOnlyFixturePatchList"
 
+namespace UE::DMX::SDMXReadOnlyFixturePatchListNamespace::Private
+{
+	class FScopedRestoreSelection
+	{
+	public:
+		FScopedRestoreSelection(const TSharedRef<SDMXReadOnlyFixturePatchList>& InFixturePatchList)
+			: FixturePatchList(InFixturePatchList)
+		{
+			SelectedItemsToRestore = FixturePatchList->GetSelectedFixturePatchRefs();
+		}
+
+		~FScopedRestoreSelection()
+		{
+			const TArray<TSharedPtr<FDMXEntityFixturePatchRef>> NewListItems = FixturePatchList->GetListItems();
+			
+			TArray<TSharedPtr<FDMXEntityFixturePatchRef>> NewSelection;
+			for (const TSharedPtr<FDMXEntityFixturePatchRef>& ItemToRestore : SelectedItemsToRestore)
+			{
+				const TSharedPtr<FDMXEntityFixturePatchRef>* CorrespondingItemPtr = Algo::FindByPredicate(NewListItems, [ItemToRestore](const TSharedPtr<FDMXEntityFixturePatchRef> NewItem)
+					{
+						return NewItem->GetFixturePatch() == ItemToRestore->GetFixturePatch();
+					});
+				if (CorrespondingItemPtr)
+				{
+					NewSelection.Add(*CorrespondingItemPtr);
+				}
+			}
+
+			FixturePatchList->SelectItems(NewSelection);
+		}
+
+	private:
+		/** Addressed fixture patch list */
+		TSharedRef<SDMXReadOnlyFixturePatchList> FixturePatchList;
+
+		/** Items to restore */
+		TArray<TSharedPtr<FDMXEntityFixturePatchRef>> SelectedItemsToRestore;
+	};
+}
+
 const FName FDMXReadOnlyFixturePatchListCollumnIDs::EditorColor = "EditorColor";
 const FName FDMXReadOnlyFixturePatchListCollumnIDs::FixturePatchName = "Name";
 const FName FDMXReadOnlyFixturePatchListCollumnIDs::FixtureID = "FID";
@@ -31,9 +72,11 @@ const FName FDMXReadOnlyFixturePatchListCollumnIDs::Patch = "Patch";
 void SDMXReadOnlyFixturePatchList::Construct(const FArguments& InArgs)
 {
 	DMXLibrary = InArgs._DMXLibrary;
+	ExcludedFixturePatches = InArgs._ExcludedFixturePatches;
 
 	IsRowEnabledDelegate = InArgs._IsRowEnabled;
 	IsRowVisibleDelegate = InArgs._IsRowVisibile;
+	OnRowDraggedDelegate = InArgs._OnRowDragged;
 
 	UpdateListItems();
 
@@ -109,11 +152,39 @@ void SDMXReadOnlyFixturePatchList::SetDMXLibrary(UDMXLibrary* InDMXLibrary)
 	RequestListRefresh();
 }
 
+void SDMXReadOnlyFixturePatchList::SetExcludedFixturePatches(const TArray<FDMXEntityFixturePatchRef>& NewExcludedFixturePatches)
+{
+	ExcludedFixturePatches = NewExcludedFixturePatches;
+	RequestListRefresh();
+}
+
+void SDMXReadOnlyFixturePatchList::SelectItems(const TArray<TSharedPtr<FDMXEntityFixturePatchRef>>& ItemsToSelect, ESelectInfo::Type SelectInfo)
+{
+	constexpr bool bSelected = true;
+	ListView->ClearSelection();
+	ListView->SetItemSelection(ItemsToSelect, bSelected, SelectInfo);
+}
+
 TArray<TSharedPtr<FDMXEntityFixturePatchRef>> SDMXReadOnlyFixturePatchList::GetSelectedFixturePatchRefs() const
 {
 	return ListView.IsValid() ? ListView->GetSelectedItems() : TArray<TSharedPtr<FDMXEntityFixturePatchRef>>();
 }
 
+TArray<TSharedPtr<FDMXEntityFixturePatchRef>> SDMXReadOnlyFixturePatchList::GetVisibleFixturePatchRefs() const
+{
+	TArray<TSharedPtr<FDMXEntityFixturePatchRef>> VisilbleFixturePatchRefs = ListItems;
+	VisilbleFixturePatchRefs.RemoveAll([this](const TSharedPtr<FDMXEntityFixturePatchRef>& FixturePatchRef)
+		{
+			if (FixturePatchRef.IsValid() && IsRowVisibleDelegate.IsBound())
+			{
+				return !IsRowVisibleDelegate.Execute(*FixturePatchRef);
+			}
+			return false;
+		});
+	
+	return VisilbleFixturePatchRefs;
+}
+
 FDMXReadOnlyFixturePatchListDescriptor SDMXReadOnlyFixturePatchList::GetListDescriptor() const
 {
 	FDMXReadOnlyFixturePatchListDescriptor ListDescriptor;
@@ -152,10 +223,13 @@ void SDMXReadOnlyFixturePatchList::RefreshList()
 {	
 	ListRefreshTimerHandle.Invalidate();
 
+	using namespace UE::DMX::SDMXReadOnlyFixturePatchListNamespace::Private;
+	const FScopedRestoreSelection ScopedRestoreSelection(StaticCastSharedRef<SDMXReadOnlyFixturePatchList>(AsShared()));
+
 	UpdateListItems();
 	ListRows.Reset(ListItems.Num());
 
-	SortByColumnID(EColumnSortPriority::Max, SortedByColumnID, EColumnSortMode::Ascending);
+	SortByColumnID(EColumnSortPriority::Max, SortedByColumnID, SortMode);
 
 	if (SearchBox.IsValid())
 	{
@@ -258,7 +332,7 @@ void SDMXReadOnlyFixturePatchList::UpdateListItems()
 	const TArray<UDMXEntityFixturePatch*> FixturePatchesInLibrary = DMXLibrary->GetEntitiesTypeCast<UDMXEntityFixturePatch>();
 	for (UDMXEntityFixturePatch* FixturePatch : FixturePatchesInLibrary)
 	{
-		if (!FixturePatch)
+		if (!FixturePatch || ExcludedFixturePatches.Contains(FixturePatch))
 		{
 			continue;
 		}
@@ -278,6 +352,7 @@ TSharedRef<ITableRow> SDMXReadOnlyFixturePatchList::OnGenerateRow(TSharedPtr<FDM
 {
 	const TSharedRef<SDMXReadOnlyFixturePatchListRow> NewRow =
 		SNew(SDMXReadOnlyFixturePatchListRow, OwnerTable, InItem.ToSharedRef())
+		.OnRowDragged(this, &SDMXReadOnlyFixturePatchList::OnRowDragged)
 		.IsEnabled(this, &SDMXReadOnlyFixturePatchList::IsRowEnabled, *InItem.Get())
 		.Visibility(TAttribute<EVisibility>::CreateSP(this, &SDMXReadOnlyFixturePatchList::GetRowVisibility, *InItem.Get()));
 
@@ -285,6 +360,16 @@ TSharedRef<ITableRow> SDMXReadOnlyFixturePatchList::OnGenerateRow(TSharedPtr<FDM
 	return NewRow;
 }
 
+FReply SDMXReadOnlyFixturePatchList::OnRowDragged(const FGeometry& MyGeometry, const FPointerEvent& MouseEvent)
+{
+	if (OnRowDraggedDelegate.IsBound())
+	{
+		return OnRowDraggedDelegate.Execute(MyGeometry, MouseEvent);
+	}
+
+	return FReply::Unhandled();
+}
+
 void SDMXReadOnlyFixturePatchList::OnSearchTextChanged(const FText& SearchText)
 {
 	RequestListRefresh();
diff --git a/Engine/Plugins/VirtualProduction/DMX/DMXEngine/Source/DMXEditor/Private/Widgets/FixturePatch/SDMXReadOnlyFixturePatchListRow.cpp b/Engine/Plugins/VirtualProduction/DMX/DMXEngine/Source/DMXEditor/Private/Widgets/FixturePatch/SDMXReadOnlyFixturePatchListRow.cpp
--- a/Engine/Plugins/VirtualProduction/DMX/DMXEngine/Source/DMXEditor/Private/Widgets/FixturePatch/SDMXReadOnlyFixturePatchListRow.cpp
+++ b/Engine/Plugins/VirtualProduction/DMX/DMXEngine/Source/DMXEditor/Private/Widgets/FixturePatch/SDMXReadOnlyFixturePatchListRow.cpp
@@ -22,6 +22,7 @@ void SDMXReadOnlyFixturePatchListRow::Construct(const FArguments& InArgs, const
 
 	SMultiColumnTableRow<TSharedPtr<FDMXEntityFixturePatchRef>>::Construct(
 		FSuperRowType::FArguments()
+		.OnDragDetected(InArgs._OnRowDragged)
 		.Style(&FDMXEditorStyle::Get().GetWidgetStyle<FTableRowStyle>("FixturePatchList.Row")),
 		InOwnerTable);
 }
diff --git a/Engine/Plugins/VirtualProduction/DMX/DMXEngine/Source/DMXEditor/Private/Widgets/FixturePatch/SDMXReadOnlyFixturePatchListRow.h b/Engine/Plugins/VirtualProduction/DMX/DMXEngine/Source/DMXEditor/Private/Widgets/FixturePatch/SDMXReadOnlyFixturePatchListRow.h
--- a/Engine/Plugins/VirtualProduction/DMX/DMXEngine/Source/DMXEditor/Private/Widgets/FixturePatch/SDMXReadOnlyFixturePatchListRow.h
+++ b/Engine/Plugins/VirtualProduction/DMX/DMXEngine/Source/DMXEditor/Private/Widgets/FixturePatch/SDMXReadOnlyFixturePatchListRow.h
@@ -20,6 +20,8 @@ public:
 	SLATE_BEGIN_ARGS(SDMXReadOnlyFixturePatchListRow)
 	{}
 
+		SLATE_EVENT(FOnDragDetected, OnRowDragged)
+
 	SLATE_END_ARGS()
 			
 	/** Constructs this widget */
diff --git a/Engine/Plugins/VirtualProduction/DMX/DMXEngine/Source/DMXEditor/Public/Widgets/SDMXReadOnlyFixturePatchList.h b/Engine/Plugins/VirtualProduction/DMX/DMXEngine/Source/DMXEditor/Public/Widgets/SDMXReadOnlyFixturePatchList.h
--- a/Engine/Plugins/VirtualProduction/DMX/DMXEngine/Source/DMXEditor/Public/Widgets/SDMXReadOnlyFixturePatchList.h
+++ b/Engine/Plugins/VirtualProduction/DMX/DMXEngine/Source/DMXEditor/Public/Widgets/SDMXReadOnlyFixturePatchList.h
@@ -77,17 +77,22 @@ struct DMXEDITOR_API FDMXReadOnlyFixturePatchListDescriptor
 class DMXEDITOR_API SDMXReadOnlyFixturePatchList
 	: public SCompoundWidget
 {
-	DECLARE_DELEGATE_RetVal_OneParam(bool, FDMXEntityFixturePatchListRowDelegate, const FDMXEntityFixturePatchRef)
+	DECLARE_DELEGATE_RetVal_OneParam(bool, FDMXEntityFixturePatchListRowDelegate, const FDMXEntityFixturePatchRef);
 
 public:
 	SLATE_BEGIN_ARGS(SDMXReadOnlyFixturePatchList)
 		: _ListDescriptor(FDMXReadOnlyFixturePatchListDescriptor())
 	{}
 
+		/** Describes the list behaviour */
 		SLATE_ARGUMENT(FDMXReadOnlyFixturePatchListDescriptor, ListDescriptor)
 
+		/** The displayed DMX Library */
 		SLATE_ARGUMENT(UDMXLibrary*, DMXLibrary)
 
+		/** Fixture patches not displayed in the list */
+		SLATE_ARGUMENT(TArray<FDMXEntityFixturePatchRef>, ExcludedFixturePatches)
+
 		/** Called when a row of the list is right clicked */
 		SLATE_EVENT(FOnContextMenuOpening, OnContextMenuOpening)
 
@@ -97,6 +102,9 @@ public:
 		/** Called to get the visibility state of each row of the list */
 		SLATE_EVENT(FDMXEntityFixturePatchListRowDelegate, IsRowVisibile)
 
+		/** Delegate executed when a row was dragged */
+		SLATE_EVENT(FOnDragDetected, OnRowDragged)
+
 	SLATE_END_ARGS()
 
 	/** Constructs this widget */
@@ -108,6 +116,9 @@ public:
 	/** Gets current selected Fixture Patches from the list */
 	TArray<TSharedPtr<FDMXEntityFixturePatchRef>> GetSelectedFixturePatchRefs() const;
 
+	/** Gets all visible Fixture Patches from the list */
+	TArray<TSharedPtr<FDMXEntityFixturePatchRef>> GetVisibleFixturePatchRefs() const;
+
 	/** Gets the a descriptor for the current parameters for this list */
 	FDMXReadOnlyFixturePatchListDescriptor GetListDescriptor() const;
 
@@ -117,6 +128,15 @@ public:
 	/** Sets the displayed DMX library */
 	void SetDMXLibrary(UDMXLibrary* InDMXLibrary);
 
+	/** Sets the excluded fixture patches */
+	void SetExcludedFixturePatches(const TArray<FDMXEntityFixturePatchRef>& NewExcludedFixturePatches);
+
+	/** Returns the items displayed int his list */
+	TArray<TSharedPtr<FDMXEntityFixturePatchRef>> GetListItems() const { return ListItems; }
+
+	/** Selects specified items in the list */
+	void SelectItems(const TArray<TSharedPtr<FDMXEntityFixturePatchRef>>& ItemsToSelect, ESelectInfo::Type SelectInfo = ESelectInfo::Direct);
+
 private:
 	/** Initializes the list parameters using the given ListDescriptor */
 	void InitializeByListDescriptor(const FDMXReadOnlyFixturePatchListDescriptor& InListDescriptor);
@@ -133,6 +153,9 @@ private:
 	/** Called when a row in the List gets generated */
 	TSharedRef<ITableRow> OnGenerateRow(TSharedPtr<FDMXEntityFixturePatchRef> InItem, const TSharedRef<STableViewBase>& OwnerTable);
 
+	/** Called when a row was dragged */
+	FReply OnRowDragged(const FGeometry& MyGeometry, const FPointerEvent& MouseEvent);
+
 	/** Called when text from searchbox is changed */
 	void OnSearchTextChanged(const FText& SearchText);
 
@@ -176,7 +199,7 @@ private:
 	bool IsUsingShowMode(EDMXReadOnlyFixturePatchListShowMode ShowModeToCheck) const;
 
 	/** The current Sort Mode */
-	EColumnSortMode::Type SortMode;
+	EColumnSortMode::Type SortMode = EColumnSortMode::Ascending;
 
 	/** The current Show Mode */
 	EDMXReadOnlyFixturePatchListShowMode ShowMode;
@@ -209,6 +232,8 @@ private:
 	TWeakObjectPtr<UDMXLibrary> DMXLibrary;
 
 	// Slate Arguments
+	TArray<FDMXEntityFixturePatchRef> ExcludedFixturePatches;
 	FDMXEntityFixturePatchListRowDelegate IsRowEnabledDelegate;
 	FDMXEntityFixturePatchListRowDelegate IsRowVisibleDelegate;
+	FOnDragDetected OnRowDraggedDelegate;
 };
diff --git a/Engine/Plugins/VirtualProduction/DMX/DMXPixelMapping/Source/DMXPixelMappingEditor/Private/Customizations/DMXPixelMappingDetailCustomization_FixtureGroup.cpp b/Engine/Plugins/VirtualProduction/DMX/DMXPixelMapping/Source/DMXPixelMappingEditor/Private/Customizations/DMXPixelMappingDetailCustomization_FixtureGroup.cpp
--- a/Engine/Plugins/VirtualProduction/DMX/DMXPixelMapping/Source/DMXPixelMappingEditor/Private/Customizations/DMXPixelMappingDetailCustomization_FixtureGroup.cpp
+++ b/Engine/Plugins/VirtualProduction/DMX/DMXPixelMapping/Source/DMXPixelMappingEditor/Private/Customizations/DMXPixelMappingDetailCustomization_FixtureGroup.cpp
@@ -7,15 +7,12 @@
 #include "Components/DMXPixelMappingFixtureGroupComponent.h"
 #include "Components/DMXPixelMappingFixtureGroupItemComponent.h"
 #include "Components/DMXPixelMappingMatrixComponent.h"
-#include "DragDrop/DMXPixelMappingDragDropOp.h"
 #include "Library/DMXEntityFixturePatch.h"
 #include "Library/DMXLibrary.h"
-#include "Templates/DMXPixelMappingComponentTemplate.h"
-#include "Toolkits/DMXPixelMappingToolkit.h"
-#include "Widgets/SDMXPixelMappingFixturePatchDetailRow.h"
 
 #include "DetailLayoutBuilder.h"
 #include "DetailWidgetRow.h"
+#include "Editor.h"
 #include "IPropertyUtilities.h"
 #include "ScopedTransaction.h"
 #include "Algo/Find.h"
@@ -45,228 +42,9 @@ void FDMXPixelMappingDetailCustomization_FixtureGroup::CustomizeDetails(IDetailL
 	// Remember the group being edited
 	WeakFixtureGroupComponent = GetSelectedFixtureGroupComponent(InDetailLayout);
 
-	// Listen to the library being changed in the group component
-	DMXLibraryHandle = InDetailLayout.GetProperty(GET_MEMBER_NAME_CHECKED(UDMXPixelMappingFixtureGroupComponent, DMXLibrary));
-	DMXLibraryHandle->SetOnPropertyValueChanged(FSimpleDelegate::CreateSP(this, &FDMXPixelMappingDetailCustomization_FixtureGroup::OnLibraryChanged));
-	DMXLibraryHandle->SetOnChildPropertyValueChanged(FSimpleDelegate::CreateSP(this, &FDMXPixelMappingDetailCustomization_FixtureGroup::OnLibraryChanged));
-
 	// Listen to component changes
 	UDMXPixelMappingBaseComponent::GetOnComponentAdded().AddSP(this, &FDMXPixelMappingDetailCustomization_FixtureGroup::OnComponentAdded);
 	UDMXPixelMappingBaseComponent::GetOnComponentRemoved().AddSP(this, &FDMXPixelMappingDetailCustomization_FixtureGroup::OnComponentRemoved);
-
-	// Add the library property
-	IDetailCategoryBuilder& FixtureListCategoryBuilder = InDetailLayout.EditCategory("Fixture List", FText::GetEmpty(), ECategoryPriority::Important);
-	FixtureListCategoryBuilder.AddProperty(DMXLibraryHandle);
-
-	// Add an 'Add all Patches' button
-	CreateAddAllPatchesButton(InDetailLayout);
-
-	// Add Fixture Patches to the view
-	CreateFixturePatchDetailRows(InDetailLayout);
-}
-
-void FDMXPixelMappingDetailCustomization_FixtureGroup::CreateAddAllPatchesButton(IDetailLayoutBuilder& InDetailLayout)
-{
-	IDetailCategoryBuilder& FixtureListCategoryBuilder = InDetailLayout.EditCategory("Fixture List", FText::GetEmpty(), ECategoryPriority::Important);
-
-	FixtureListCategoryBuilder.AddCustomRow(FText::GetEmpty())
-		.WholeRowContent()
-		[
-			SNew(SButton)
-			.Text(LOCTEXT("AddAllPatchesButtonText", "Add all Patches"))
-			.OnClicked(this, &FDMXPixelMappingDetailCustomization_FixtureGroup::OnAddAllPatchesClicked)
-			.Visibility_Lambda([this]()
-				{
-					return NumFixturePatchRows > 0 ? EVisibility::Visible : EVisibility::Collapsed;
-				})
-		];
-}
-
-void FDMXPixelMappingDetailCustomization_FixtureGroup::CreateFixturePatchDetailRows(IDetailLayoutBuilder& InDetailLayout)
-{
-	
-
-	UDMXPixelMappingFixtureGroupComponent* FixtureGroupComponent = WeakFixtureGroupComponent.Get();
-	if (!FixtureGroupComponent)
-	{
-		return;
-	}
-
-	UDMXLibrary* DMXLibrary = GetSelectedDMXLibrary(FixtureGroupComponent);
-	if (!DMXLibrary)
-	{
-		return;
-	}
-
-	UpdateFixturePatchesInUse(DMXLibrary);
-
-	// Listen to the entities array being changed in the library
-	EntitiesHandle = InDetailLayout.GetProperty(UDMXLibrary::GetEntitiesPropertyName());
-	EntitiesHandle->SetOnPropertyValueChanged(FSimpleDelegate::CreateSP(this, &FDMXPixelMappingDetailCustomization_FixtureGroup::ForceRefresh));
-
-	// Add fixture patches as custom rows
-	NumFixturePatchRows = 0;
-	TArray<UDMXEntityFixturePatch*> AllFixturePatches = DMXLibrary->GetEntitiesTypeCast<UDMXEntityFixturePatch>();
-	for (UDMXEntityFixturePatch* FixturePatch : AllFixturePatches)
-	{
-		const bool bPatchIsAssigned = IsFixturePatchAssignedToPixelMapping(FixturePatch);
-
-		if (!bPatchIsAssigned)
-		{
-			TSharedRef<SDMXPixelMappingFixturePatchDetailRow> FixturePatchDetailRowWidget =
-				SNew(SDMXPixelMappingFixturePatchDetailRow)
-				.FixturePatch(FixturePatch)
-				.OnLMBDown(this, &FDMXPixelMappingDetailCustomization_FixtureGroup::OnFixturePatchLMBDown, FDMXEntityFixturePatchRef(FixturePatch))
-				.OnLMBUp(this, &FDMXPixelMappingDetailCustomization_FixtureGroup::OnFixturePatchLMBUp, FDMXEntityFixturePatchRef(FixturePatch))
-				.OnDragged(this, &FDMXPixelMappingDetailCustomization_FixtureGroup::OnFixturePatchesDragged);
-
-			FDMXPixelMappingDetailCustomization_FixtureGroup::FDetailRowWidgetWithPatch DetailRowWidgetWithPatch;
-			DetailRowWidgetWithPatch.DetailRowWidget = FixturePatchDetailRowWidget;
-			DetailRowWidgetWithPatch.WeakFixturePatch = FixturePatch;
-			DetailRowWidgetsWithPatch.Add(DetailRowWidgetWithPatch);
-
-			IDetailCategoryBuilder& FixtureListCategoryBuilder = InDetailLayout.EditCategory("Fixture List", FText::GetEmpty(), ECategoryPriority::Important);
-			FixtureListCategoryBuilder.AddCustomRow(FText::GetEmpty())
-				.WholeRowContent()
-				[
-					FixturePatchDetailRowWidget
-				];
-
-			NumFixturePatchRows++;
-		}
-	}
-}
-
-FReply FDMXPixelMappingDetailCustomization_FixtureGroup::OnAddAllPatchesClicked()
-{
-	const TSharedPtr<FDMXPixelMappingToolkit> Toolkit = ToolkitWeakPtr.Pin();
-	if (!Toolkit.IsValid())
-	{
-		return FReply::Handled();
-	}
-
-	UDMXPixelMappingFixtureGroupComponent* GroupComponent = WeakFixtureGroupComponent.Get();
-	if (!GroupComponent)
-	{
-		return FReply::Handled();
-	}
-
-	UDMXLibrary* DMXLibrary = GetSelectedDMXLibrary(GroupComponent);
-	if (!DMXLibrary)
-	{
-		return FReply::Handled();
-	}
-
-	UDMXPixelMapping* PixelMapping = GroupComponent->GetPixelMapping();
-	if (!PixelMapping)
-	{
-		return FReply::Handled();
-	}
-
-	UDMXPixelMappingRootComponent* RootComponent = PixelMapping->GetRootComponent();
-	if (!RootComponent)
-	{
-		return FReply::Handled();
-	}
-
-	// Create components
-	FScopedTransaction AddAllPatchesTransaction(LOCTEXT("AddAllPatchesTransaction", "Add all Patches to Pixel Mapping"));
-	GroupComponent->Modify();
-
-	TArray<UDMXPixelMappingOutputComponent*> NewComponents;
-	for (const FDMXEntityFixturePatchRef& FixturePatchRef : FixturePatches)
-	{
-		UDMXEntityFixturePatch* FixturePatch = FixturePatchRef.GetFixturePatch();
-		if (!FixturePatch)
-		{
-			continue;
-		}
-
-		const bool bPatchIsAssigned = IsFixturePatchAssignedToPixelMapping(FixturePatch);
-		if (bPatchIsAssigned)
-		{
-			continue;
-		}
-
-		const FDMXFixtureMode* FixtureModePtr = FixturePatch->GetActiveMode();
-		if (FixtureModePtr && FixtureModePtr->bFixtureMatrixEnabled)
-		{
-			const TSharedPtr<FDMXPixelMappingComponentTemplate> ComponentTemplate = MakeShared<FDMXPixelMappingComponentTemplate>(UDMXPixelMappingMatrixComponent::StaticClass(), FixturePatchRef);
-			UDMXPixelMappingMatrixComponent* NewMatrixComponent = ComponentTemplate->CreateComponent<UDMXPixelMappingMatrixComponent>(RootComponent);
-			NewComponents.Add(NewMatrixComponent);
-
-			GroupComponent->AddChild(NewMatrixComponent);
-		}
-		else
-		{
-			const TSharedPtr<FDMXPixelMappingComponentTemplate> ComponentTemplate = MakeShared<FDMXPixelMappingComponentTemplate>(UDMXPixelMappingFixtureGroupItemComponent::StaticClass(), FixturePatchRef);
-			UDMXPixelMappingFixtureGroupItemComponent* NewGroupItemComponent = ComponentTemplate->CreateComponent<UDMXPixelMappingFixtureGroupItemComponent>(RootComponent);
-			NewComponents.Add(NewGroupItemComponent);
-
-			GroupComponent->AddChild(NewGroupItemComponent);
-		}
-	}
-
-	// Layout new components inside the group
-	const int32 Columns = FMath::RoundFromZero(FMath::Sqrt((float)NewComponents.Num()));
-	const int32 Rows = FMath::RoundFromZero((float)NewComponents.Num() / Columns);
-	const FVector2D Size = FVector2D(GroupComponent->GetSize().X / Columns, GroupComponent->GetSize().Y / Rows);
-		
-	const FVector2D ParentPosition = GroupComponent->GetPosition();
-	int32 Column = -1;
-	int32 Row = 0;
-	for (UDMXPixelMappingOutputComponent* NewComponent : NewComponents)
-	{
-		Column++;
-		if (Column >= Columns)
-		{
-			Column = 0;
-			Row++;
-		}
-		const FVector2D Position = ParentPosition + FVector2D(Column * Size.X, Row * Size.Y);
-		NewComponent->SetPosition(Position);
-		NewComponent->SetSize(Size);
-	}
-
-	return FReply::Handled();
-}
-
-void FDMXPixelMappingDetailCustomization_FixtureGroup::OnLibraryChanged()
-{
-	// Remove patches not in the library
-	if (TSharedPtr<FDMXPixelMappingToolkit> Toolkit = ToolkitWeakPtr.Pin())
-	{
-		if (UDMXPixelMappingFixtureGroupComponent* GroupComponent = Cast<UDMXPixelMappingFixtureGroupComponent>(WeakFixtureGroupComponent))
-		{
-			const FScopedTransaction ChangeDMXLibraryTransaction(LOCTEXT("DMXLibraryChangedResetTransactionReason", "PixelMapping Changed DMX Library"));
-
-			TArray<UDMXPixelMappingBaseComponent*> CachedChildren(GroupComponent->Children);
-			for (UDMXPixelMappingBaseComponent* ChildComponent : CachedChildren)
-			{
-				if (UDMXPixelMappingFixtureGroupItemComponent* ChildGroupItem = Cast<UDMXPixelMappingFixtureGroupItemComponent>(ChildComponent))
-				{
-					if (ChildGroupItem->FixturePatchRef.GetFixturePatch() &&
-						ChildGroupItem->FixturePatchRef.GetFixturePatch()->GetParentLibrary() != GroupComponent->DMXLibrary)
-					{
-						GroupComponent->RemoveChild(ChildGroupItem);
-					}
-				}
-				else if (UDMXPixelMappingMatrixComponent* ChildMatrix = Cast<UDMXPixelMappingMatrixComponent>(ChildComponent))
-				{
-					if (ChildMatrix->FixturePatchRef.GetFixturePatch() &&
-						ChildMatrix->FixturePatchRef.GetFixturePatch()->GetParentLibrary() != GroupComponent->DMXLibrary)
-					{
-						GroupComponent->RemoveChild(ChildMatrix);
-					}
-				}
-			};
-		}
-	}
-
-	if (!RequestForceRefreshHandle.IsValid())
-	{
-		RequestForceRefreshHandle = FCoreDelegates::OnEndFrame.AddSP(this, &FDMXPixelMappingDetailCustomization_FixtureGroup::ForceRefresh);
-	}
 }
 
 void FDMXPixelMappingDetailCustomization_FixtureGroup::OnComponentAdded(UDMXPixelMapping* PixelMapping, UDMXPixelMappingBaseComponent* Component)
@@ -388,9 +166,6 @@ void FDMXPixelMappingDetailCustomization_FixtureGroup::HandleSizePropertyChanged
 void FDMXPixelMappingDetailCustomization_FixtureGroup::ForceRefresh()
 {
 	// Reset the handles so they won't fire any changes after refreshing
-	DMXLibraryHandle.Reset();
-	EntitiesHandle.Reset();
-
 	if (ensure(PropertyUtilities.IsValid()))
 	{
 		PropertyUtilities->ForceRefresh();
@@ -399,178 +174,6 @@ void FDMXPixelMappingDetailCustomization_FixtureGroup::ForceRefresh()
 	RequestForceRefreshHandle.Reset();
 }
 
-void FDMXPixelMappingDetailCustomization_FixtureGroup::OnFixturePatchLMBDown(const FGeometry& MyGeometry, const FPointerEvent& MouseEvent, FDMXEntityFixturePatchRef FixturePatchRef)
-{
-	UDMXEntityFixturePatch* FixturePatch = FixturePatchRef.GetFixturePatch();
-	if (FixturePatch)
-	{
-		if (SelectedFixturePatches.Num() == 0)
-		{
-			SelectedFixturePatches = TArray<FDMXEntityFixturePatchRef>({ FixturePatchRef });
-		}
-		else
-		{
-			if (MouseEvent.IsShiftDown())
-			{
-				// Shift select
-				const int32 IndexOfAnchor = DetailRowWidgetsWithPatch.IndexOfByPredicate([this](const FDMXPixelMappingDetailCustomization_FixtureGroup::FDetailRowWidgetWithPatch& DetailRowWidgetWithPatch) {
-					return
-						DetailRowWidgetWithPatch.WeakFixturePatch.IsValid() &&
-						SelectedFixturePatches[0].GetFixturePatch() &&
-						DetailRowWidgetWithPatch.WeakFixturePatch.Get() == SelectedFixturePatches[0].GetFixturePatch();
-					});
-
-				const int32 IndexOfSelected = DetailRowWidgetsWithPatch.IndexOfByPredicate([FixturePatch](const FDMXPixelMappingDetailCustomization_FixtureGroup::FDetailRowWidgetWithPatch& DetailRowWidgetWithPatch) {
-					return
-						DetailRowWidgetWithPatch.WeakFixturePatch.IsValid() &&
-						DetailRowWidgetWithPatch.WeakFixturePatch.Get() == FixturePatch;
-					});
-
-				if (ensure(IndexOfSelected != INDEX_NONE))
-				{
-					if (IndexOfAnchor == INDEX_NONE || IndexOfAnchor == IndexOfSelected)
-					{
-						SelectedFixturePatches = TArray<FDMXEntityFixturePatchRef>({ FixturePatchRef });
-					}
-					else
-					{
-						SelectedFixturePatches.Reset();
-						SelectedFixturePatches.Add(FDMXEntityFixturePatchRef(DetailRowWidgetsWithPatch[IndexOfAnchor].WeakFixturePatch.Get()));
-
-						const bool bAscending = IndexOfAnchor < IndexOfSelected;
-						const int32 StartIndex = bAscending ? IndexOfAnchor + 1 : IndexOfSelected;
-						const int32 EndIndex = bAscending ? IndexOfSelected : IndexOfAnchor - 1;
-
-						for (int32 IndexDetailRowWidget = StartIndex; IndexDetailRowWidget <= EndIndex; IndexDetailRowWidget++)
-						{
-							if (UDMXEntityFixturePatch* NewlySelectedPatch = DetailRowWidgetsWithPatch[IndexDetailRowWidget].WeakFixturePatch.Get())
-							{
-								SelectedFixturePatches.AddUnique(FDMXEntityFixturePatchRef(NewlySelectedPatch));
-							}
-						}
-					}
-				}
-			}
-			else if (MouseEvent.IsControlDown())
-			{
-				// Ctrl select
-				if (!SelectedFixturePatches.Contains(FixturePatchRef))
-				{
-					SelectedFixturePatches.Add(FixturePatchRef);
-				}
-				else
-				{
-					SelectedFixturePatches.Remove(FixturePatchRef);
-				}
-			}
-		}
-
-		UpdateFixturePatchHighlights();
-	}
-}
-
-void FDMXPixelMappingDetailCustomization_FixtureGroup::OnFixturePatchLMBUp(const FGeometry& MyGeometry, const FPointerEvent& MouseEvent, FDMXEntityFixturePatchRef FixturePatchRef)
-{
-	if (!MouseEvent.IsShiftDown() && !MouseEvent.IsControlDown())
-	{
-		// Make a new selection
-		SelectedFixturePatches.Reset();
-		SelectedFixturePatches.Add(FixturePatchRef);
-
-		UpdateFixturePatchHighlights();
-	}
-}
-
-FReply FDMXPixelMappingDetailCustomization_FixtureGroup::OnFixturePatchesDragged(const FGeometry& MyGeometry, const FPointerEvent& MouseEvent)
-{
-	if (WeakFixtureGroupComponent.IsValid())
-	{
-		TArray<TSharedPtr<FDMXPixelMappingComponentTemplate>> Templates;
-		for (const FDMXEntityFixturePatchRef& FixturePatchRef : SelectedFixturePatches)
-		{
-			const UDMXEntityFixturePatch* FixturePatch = FixturePatchRef.GetFixturePatch();
-			const UDMXEntityFixtureType* FixtureType = FixturePatch ? FixturePatch->GetFixtureType() : nullptr;
-			const FDMXFixtureMode* ActiveModePtr = FixturePatch ? FixturePatch->GetActiveMode() : nullptr;
-			if (FixturePatch && FixtureType && ActiveModePtr)
-			{
-				if (ActiveModePtr->bFixtureMatrixEnabled)
-				{
-					TSharedRef<FDMXPixelMappingComponentTemplate> FixturePatchMatrixTemplate = MakeShared<FDMXPixelMappingComponentTemplate>(UDMXPixelMappingMatrixComponent::StaticClass(), FixturePatchRef);
-					Templates.Add(FixturePatchMatrixTemplate);
-				}
-				else
-				{
-					TSharedRef<FDMXPixelMappingComponentTemplate> FixturePatchItemTemplate = MakeShared<FDMXPixelMappingComponentTemplate>(UDMXPixelMappingFixtureGroupItemComponent::StaticClass(), FixturePatchRef);
-					Templates.Add(FixturePatchItemTemplate);
-				}
-			}
-		}
-
-		UpdateFixturePatchHighlights();
-
-		ForceRefresh();
-
-		return FReply::Handled().BeginDragDrop(FDMXPixelMappingDragDropOp::New(FVector2D::ZeroVector, Templates, WeakFixtureGroupComponent.Get()));
-	}
-
-	return FReply::Handled();
-}
-
-bool FDMXPixelMappingDetailCustomization_FixtureGroup::IsFixturePatchAssignedToPixelMapping(UDMXEntityFixturePatch* FixturePatch) const
-{
-	if (!WeakFixtureGroupComponent.IsValid())
-	{
-		return true;
-	}
-
-	const bool bPatchIsAssigned = WeakFixtureGroupComponent->Children.ContainsByPredicate([FixturePatch](UDMXPixelMappingBaseComponent* BaseComponent)
-		{
-			if (UDMXPixelMappingFixtureGroupItemComponent* GroupItemComponent = Cast<UDMXPixelMappingFixtureGroupItemComponent>(BaseComponent))
-			{
-				return GroupItemComponent->FixturePatchRef.GetFixturePatch() == FixturePatch;
-			}
-			else if (UDMXPixelMappingMatrixComponent* MatrixComponent = Cast<UDMXPixelMappingMatrixComponent>(BaseComponent))
-			{
-				return MatrixComponent->FixturePatchRef.GetFixturePatch() == FixturePatch;
-			}
-			return false;
-		});
-
-	return bPatchIsAssigned;
-}
-
-void FDMXPixelMappingDetailCustomization_FixtureGroup::UpdateFixturePatchHighlights()
-{
-	for (const FDMXPixelMappingDetailCustomization_FixtureGroup::FDetailRowWidgetWithPatch& DetailRowWidgetWithPatch : DetailRowWidgetsWithPatch)
-	{
-		bool bIsSelected = SelectedFixturePatches.ContainsByPredicate([DetailRowWidgetWithPatch](const FDMXEntityFixturePatchRef& SelectedRef) {
-			return
-				SelectedRef.GetFixturePatch() &&
-				DetailRowWidgetWithPatch.WeakFixturePatch.IsValid() &&
-				SelectedRef.GetFixturePatch() == DetailRowWidgetWithPatch.WeakFixturePatch.Get();
-			});
-
-		DetailRowWidgetWithPatch.DetailRowWidget->SetHighlight(bIsSelected);
-	}
-}
-
-void FDMXPixelMappingDetailCustomization_FixtureGroup::UpdateFixturePatchesInUse(UDMXLibrary* DMXLibrary)
-{
-	TArray<UDMXEntityFixturePatch*> FixturePatchesInLibrary = DMXLibrary->GetEntitiesTypeCast<UDMXEntityFixturePatch>();
-
-	FixturePatches.Reset(FixturePatchesInLibrary.Num());
-	for (UDMXEntityFixturePatch* FixturePatch : FixturePatchesInLibrary)
-	{
-		FDMXEntityFixturePatchRef FixturePatchRef;
-		FixturePatchRef.SetEntity(FixturePatch);
-		FixturePatches.Add(FixturePatchRef);
-	}
-
-	SelectedFixturePatches.RemoveAll([&FixturePatchesInLibrary](const FDMXEntityFixturePatchRef& SelectedFixturePatch) {
-		return !SelectedFixturePatch.GetFixturePatch() || !FixturePatchesInLibrary.Contains(SelectedFixturePatch.GetFixturePatch());
-		});
-}
-
 void FDMXPixelMappingDetailCustomization_FixtureGroup::UpdateCachedScaleChildrenWithParent()
 {
 	const UDMXPixelMappingLayoutSettings* LayoutSettings = GetDefault<UDMXPixelMappingLayoutSettings>();
diff --git a/Engine/Plugins/VirtualProduction/DMX/DMXPixelMapping/Source/DMXPixelMappingEditor/Private/Customizations/DMXPixelMappingDetailCustomization_FixtureGroup.h b/Engine/Plugins/VirtualProduction/DMX/DMXPixelMapping/Source/DMXPixelMappingEditor/Private/Customizations/DMXPixelMappingDetailCustomization_FixtureGroup.h
--- a/Engine/Plugins/VirtualProduction/DMX/DMXPixelMapping/Source/DMXPixelMappingEditor/Private/Customizations/DMXPixelMappingDetailCustomization_FixtureGroup.h
+++ b/Engine/Plugins/VirtualProduction/DMX/DMXPixelMapping/Source/DMXPixelMappingEditor/Private/Customizations/DMXPixelMappingDetailCustomization_FixtureGroup.h
@@ -7,23 +7,15 @@
 
 struct FPropertyChangedEvent;
 
-struct FDMXEntityFixturePatchRef;
 class FDMXPixelMappingToolkit;
-class SDMXPixelMappingFixturePatchDetailRow;
-class UDMXEntityFixturePatch;
 class UDMXLibrary;
 class UDMXPixelMapping;
 class UDMXPixelMappingBaseComponent;
 class UDMXPixelMappingFixtureGroupComponent;
 
-class FReply;
-struct FPointerEvent;
-struct FGeometry;
 class IDetailLayoutBuilder;
 class IPropertyHandle;
 class IPropertyUtilities;
-class SBorder;
-template <typename ItemType> class SListView;
 
 
 class FDMXPixelMappingDetailCustomization_FixtureGroup
@@ -44,18 +36,6 @@ public:
 	virtual void CustomizeDetails(IDetailLayoutBuilder& DetailLayout) override;
 
 private:
-	/** Creates a button that adds all patches to the group */
-	void CreateAddAllPatchesButton(IDetailLayoutBuilder& InDetailLayout);
-
-	/** Creates a list of fixture patches available to drag drop as detail rws */
-	void CreateFixturePatchDetailRows(IDetailLayoutBuilder& InDetailLayout);
-
-	/** Called when the 'Add all Patches' button was clicked */
-	FReply OnAddAllPatchesClicked();
-
-	/** Called when the library changed */
-	void OnLibraryChanged();
-
 	/** Called when a component was added */
 	void OnComponentAdded(UDMXPixelMapping* PixelMapping, UDMXPixelMappingBaseComponent* Component);
 
@@ -74,24 +54,6 @@ private:
 	/** Forces the detail layout to refresh */
 	void ForceRefresh();
 
-	/** Called when a fixture patch widget got a mouse down event */
-	void OnFixturePatchLMBDown(const FGeometry& MyGeometry, const FPointerEvent& MouseEvent, FDMXEntityFixturePatchRef FixturePatchRef);
-
-	/** Called when a fixture patch widget got a mouse up event */
-	void OnFixturePatchLMBUp(const FGeometry& MyGeometry, const FPointerEvent& MouseEvent, FDMXEntityFixturePatchRef FixturePatchRef);
-
-	/** Called when a fixture patch was dragged */
-	FReply OnFixturePatchesDragged(const FGeometry& MyGeometry, const FPointerEvent& MouseEvent);
-
-	/** Returns true if the fixture patch is used in the pixel mapping */
-	bool IsFixturePatchAssignedToPixelMapping(UDMXEntityFixturePatch* FixturePatch) const;
-
-	/** Updates highlights for the fixture patches (selection) */
-	void UpdateFixturePatchHighlights();
-
-	/** Update fixture patches in use */
-	void UpdateFixturePatchesInUse(UDMXLibrary* DMXLibrary);
-
 	/** Updates the bCachedScaleChildrenWithParent member from UDMXPixelMappingLayoutSettings */
 	void UpdateCachedScaleChildrenWithParent();
 
@@ -104,44 +66,21 @@ private:
 	/** Weak reference to the DMX editor */
 	TWeakPtr<FDMXPixelMappingToolkit> ToolkitWeakPtr;
 
-	/** Fixture patches currently in the library */
-	TArray<FDMXEntityFixturePatchRef> FixturePatches;
-
-	/** Fixture patches currently in the library and selected */
-	TArray<FDMXEntityFixturePatchRef> SelectedFixturePatches;
-
 	/** The single fixture group component in use */
 	TWeakObjectPtr<UDMXPixelMappingFixtureGroupComponent> WeakFixtureGroupComponent;
 
-	/** Number of fixture patch rows displayed */
-	int32 NumFixturePatchRows = 0;
-
 	/** SizeX before it got changed */
 	TMap<TWeakObjectPtr<UDMXPixelMappingFixtureGroupComponent>, FVector2D> PreEditChangeComponentToSizeMap;
 
 	/** Cached UDMXPixelMappingLayoutSettings::bScaleChildrenWithParen, to avoid repetitive reads during interactive changes */
 	bool bCachedScaleChildrenWithParent = false;
 
-	/** Handle to the dmx library property */
-	TSharedPtr<IPropertyHandle> DMXLibraryHandle;
-
-	/** Handle to the dmx library's entity array */
-	TSharedPtr<IPropertyHandle> EntitiesHandle;
-
 	/** Property handle for the SizeX property */
 	TSharedPtr<IPropertyHandle> SizeXHandle;
 
 	/** Property handle for the SizeY property */
 	TSharedPtr<IPropertyHandle> SizeYHandle;
 
-	/** Patches with their details row. Helps to multiselect and toggle highlights */
-	struct FDetailRowWidgetWithPatch
-	{
-		TWeakObjectPtr<UDMXEntityFixturePatch> WeakFixturePatch;
-		TSharedPtr<SDMXPixelMappingFixturePatchDetailRow> DetailRowWidget;
-	};
-	TArray<FDetailRowWidgetWithPatch> DetailRowWidgetsWithPatch;
-
 	/** Delegate handle while being bound to the OnEndFrame event */
 	FDelegateHandle RequestForceRefreshHandle;
 
diff --git a/Engine/Plugins/VirtualProduction/DMX/DMXPixelMapping/Source/DMXPixelMappingEditor/Private/Toolkits/DMXPixelMappingToolkit.cpp b/Engine/Plugins/VirtualProduction/DMX/DMXPixelMapping/Source/DMXPixelMappingEditor/Private/Toolkits/DMXPixelMappingToolkit.cpp
--- a/Engine/Plugins/VirtualProduction/DMX/DMXPixelMapping/Source/DMXPixelMappingEditor/Private/Toolkits/DMXPixelMappingToolkit.cpp
+++ b/Engine/Plugins/VirtualProduction/DMX/DMXPixelMapping/Source/DMXPixelMappingEditor/Private/Toolkits/DMXPixelMappingToolkit.cpp
@@ -17,12 +17,12 @@
 #include "Components/DMXPixelMappingMatrixComponent.h"
 #include "Components/DMXPixelMappingScreenComponent.h"
 #include "Templates/DMXPixelMappingComponentTemplate.h"
-#include "ViewModels/DMXPixelMappingPaletteViewModel.h"
 #include "Views/SDMXPixelMappingDesignerView.h"
 #include "Views/SDMXPixelMappingDetailsView.h"
+#include "Views/SDMXPixelMappingDMXLibraryView.h"
 #include "Views/SDMXPixelMappingHierarchyView.h"
+#include "Views/SDMXPixelMappingInputSourceView.h"
 #include "Views/SDMXPixelMappingLayoutView.h"
-#include "Views/SDMXPixelMappingPaletteView.h"
 #include "Views/SDMXPixelMappingPreviewView.h"
 
 #include "ScopedTransaction.h"
@@ -36,7 +36,8 @@
 
 #define LOCTEXT_NAMESPACE "DMXPixelMappingToolkit"
 
-const FName FDMXPixelMappingToolkit::PaletteViewTabID(TEXT("DMXPixelMappingEditor_PaletteViewTabID"));
+const FName FDMXPixelMappingToolkit::InputSourceViewTabID(TEXT("DMXPixelMappingEditor_InputSourceViewTabID"));
+const FName FDMXPixelMappingToolkit::DMXLibraryViewTabID(TEXT("DMXPixelMappingEditor_DMXLibraryViewTabID"));
 const FName FDMXPixelMappingToolkit::HierarchyViewTabID(TEXT("DMXPixelMappingEditor_HierarchyViewTabID"));
 const FName FDMXPixelMappingToolkit::DesignerViewTabID(TEXT("DMXPixelMappingEditor_DesignerViewTabID"));
 const FName FDMXPixelMappingToolkit::PreviewViewTabID(TEXT("DMXPixelMappingEditor_PreviewViewTabID"));
@@ -73,8 +74,13 @@
 
 	FAssetEditorToolkit::RegisterTabSpawners(InTabManager);
 
-	InTabManager->RegisterTabSpawner(PaletteViewTabID, FOnSpawnTab::CreateSP(this, &FDMXPixelMappingToolkit::SpawnTab_PaletteView))
-		.SetDisplayName(LOCTEXT("Tab_PaletteView", "Palette"))
+	InTabManager->RegisterTabSpawner(InputSourceViewTabID, FOnSpawnTab::CreateSP(this, &FDMXPixelMappingToolkit::SpawnTab_InputSourceView))
+		.SetDisplayName(LOCTEXT("Tab_InputSourceView", "Input Source"))
 		.SetGroup(WorkspaceMenuCategoryRef)
 		.SetIcon(FSlateIcon(FDMXPixelMappingEditorStyle::GetStyleSetName(), "ClassIcon.DMXPixelMapping"));
 
+	InTabManager->RegisterTabSpawner(DMXLibraryViewTabID, FOnSpawnTab::CreateSP(this, &FDMXPixelMappingToolkit::SpawnTab_DMXLibraryView))
+		.SetDisplayName(LOCTEXT("Tab_DMXLibraryView", "DMX Library"))
+		.SetGroup(WorkspaceMenuCategoryRef)
+		.SetIcon(FSlateIcon(FDMXPixelMappingEditorStyle::GetStyleSetName(), "ClassIcon.DMXPixelMapping"));
+
@@ -108,7 +114,7 @@
 {
 	FAssetEditorToolkit::UnregisterTabSpawners(InTabManager);
 
-	InTabManager->UnregisterTabSpawner(PaletteViewTabID);
+	InTabManager->UnregisterTabSpawner(InputSourceViewTabID);
 	InTabManager->UnregisterTabSpawner(HierarchyViewTabID);
 	InTabManager->UnregisterTabSpawner(DesignerViewTabID);
 	InTabManager->UnregisterTabSpawner(PreviewViewTabID);
@@ -292,13 +298,10 @@
 
 void FDMXPixelMappingToolkit::AddRenderer()
 {
-	if (DMXPixelMapping != nullptr)
+	if (DMXPixelMapping)
 	{
-		{
-			// Just use root component for now
-			UDMXPixelMappingRendererComponent* RendererComponent = FDMXPixelMappingEditorUtils::AddRenderer(DMXPixelMapping);
-			SetActiveRenderComponent(RendererComponent);
-		}
+		UDMXPixelMappingRendererComponent* RendererComponent = FDMXPixelMappingEditorUtils::AddRenderer(DMXPixelMapping);
+		SetActiveRenderComponent(RendererComponent);
 	}
 }
 
@@ -317,102 +320,6 @@
 	bTogglePlayDMXAll = !bTogglePlayDMXAll;
 }
 
-bool FDMXPixelMappingToolkit::CanSizeSelectedComponentToTexture() const
-{
-	if (SelectedComponents.Num() == 1)
-	{
-		if (UDMXPixelMappingBaseComponent* Component = SelectedComponents.Array()[0].GetComponent())
-		{
-			return
-				Component->GetClass() == UDMXPixelMappingFixtureGroupComponent::StaticClass() ||
-				Component->GetClass() == UDMXPixelMappingMatrixComponent::StaticClass() ||
-				Component->GetClass() == UDMXPixelMappingScreenComponent::StaticClass() ||
-				(Component->GetParent() && Component->GetParent()->GetClass() == UDMXPixelMappingFixtureGroupComponent::StaticClass()) ||
-				(Component->GetParent() && Component->GetParent()->GetClass() == UDMXPixelMappingMatrixComponent::StaticClass()) ||
-				(Component->GetParent() && Component->GetParent()->GetClass() == UDMXPixelMappingScreenComponent::StaticClass());
-		}
-	}
-	return false;
-}
-
-void FDMXPixelMappingToolkit::SizeSelectedComponentToTexture()
-{
-	if (SelectedComponents.IsEmpty())
-	{
-		return;
-	}	
-
-	const UDMXPixelMappingRendererComponent* RendererComponent = GetActiveRendererComponent();
-	if (!RendererComponent)
-	{
-		return;
-	}
-
-	const FVector2D TextureSize = RendererComponent->GetSize();
-	if (TextureSize == FVector2D::ZeroVector)
-	{
-		return;
-	}
-
-	UDMXPixelMappingOutputComponent* Component = Cast<UDMXPixelMappingOutputComponent>(SelectedComponents.Array()[0].GetComponent());
-	if (!Component)
-	{
-		return;
-	}
-
-	// Apply to parent where appropriate
-	if (Component->GetClass() != UDMXPixelMappingFixtureGroupComponent::StaticClass() &&
-		Component->GetClass() != UDMXPixelMappingMatrixComponent::StaticClass() &&
-		Component->GetClass() != UDMXPixelMappingScreenComponent::StaticClass())
-	{
-		if ((Component->GetParent() && Component->GetParent()->GetClass() == UDMXPixelMappingFixtureGroupComponent::StaticClass()) ||
-			(Component->GetParent() && Component->GetParent()->GetClass() == UDMXPixelMappingMatrixComponent::StaticClass()) ||
-			(Component->GetParent() && Component->GetParent()->GetClass() == UDMXPixelMappingScreenComponent::StaticClass()))
-		{
-			Component = Cast<UDMXPixelMappingOutputComponent>(Component->GetParent());
-		}
-	}
-
-	const FScopedTransaction SizeComponentToTextureTransaction(LOCTEXT("SizeComponentToTextureTransaction", "Size Component to Texture"));
-
-	const UDMXPixelMappingLayoutSettings* LayoutSettings = GetDefault<UDMXPixelMappingLayoutSettings>();
-	if (LayoutSettings && LayoutSettings->bScaleChildrenWithParent)
-	{		
-		// Scale children to parent
-		const FVector2D RatioVector = TextureSize / Component->GetSize();
-		for (UDMXPixelMappingBaseComponent* BaseChild : Component->GetChildren())
-		{
-			if (UDMXPixelMappingOutputComponent* Child = Cast<UDMXPixelMappingOutputComponent>(BaseChild))
-			{
-				Child->PreEditChange(UDMXPixelMappingOutputComponent::StaticClass()->FindPropertyByName(UDMXPixelMappingOutputComponent::GetPositionXPropertyName()));
-				Child->PreEditChange(UDMXPixelMappingOutputComponent::StaticClass()->FindPropertyByName(UDMXPixelMappingOutputComponent::GetPositionYPropertyName()));
-				Child->PreEditChange(UDMXPixelMappingOutputComponent::StaticClass()->FindPropertyByName(UDMXPixelMappingOutputComponent::GetSizeXPropertyName()));
-				Child->PreEditChange(UDMXPixelMappingOutputComponent::StaticClass()->FindPropertyByName(UDMXPixelMappingOutputComponent::GetSizeYPropertyName()));
-
-				// Scale size (SetSize already clamps)
-				Child->SetSize(Child->GetSize() * RatioVector);
-
-				// Scale position (new position is zero vector)
-				const FVector2D ChildPosition = Child->GetPosition();
-				const FVector2D NewPositionRelative = (ChildPosition - Component->GetPosition()) * RatioVector;
-				Child->SetPosition(Component->GetPosition() + NewPositionRelative);
-
-				Child->PostEditChange();
-			}
-		}
-	}
-
-	Component->PreEditChange(UDMXPixelMappingOutputComponent::StaticClass()->FindPropertyByName(UDMXPixelMappingOutputComponent::GetPositionXPropertyName()));
-	Component->PreEditChange(UDMXPixelMappingOutputComponent::StaticClass()->FindPropertyByName(UDMXPixelMappingOutputComponent::GetPositionYPropertyName()));
-	Component->PreEditChange(UDMXPixelMappingOutputComponent::StaticClass()->FindPropertyByName(UDMXPixelMappingOutputComponent::GetSizeXPropertyName()));
-	Component->PreEditChange(UDMXPixelMappingOutputComponent::StaticClass()->FindPropertyByName(UDMXPixelMappingOutputComponent::GetSizeYPropertyName()));
-
-	Component->SetPosition(FVector2D::ZeroVector);
-	Component->SetSize(TextureSize);
-
-	Component->PostEditChange();
-}
-
 void FDMXPixelMappingToolkit::UpdateBlueprintNodes(UDMXPixelMapping* InDMXPixelMapping)
 {
 	if (InDMXPixelMapping != nullptr)
@@ -508,6 +415,120 @@
 	}
 }
 
+void FDMXPixelMappingToolkit::ZoomToFit()
+{
+	if (DesignerView.IsValid())
+	{
+		constexpr bool bInstantZoom = false;
+		DesignerView->ZoomToFit(bInstantZoom);
+	}
+}
+
+bool FDMXPixelMappingToolkit::CanSizeSelectedComponentToTexture() const
+{
+	if (SelectedComponents.Num() == 1)
+	{
+		if (UDMXPixelMappingBaseComponent* Component = SelectedComponents.Array()[0].GetComponent())
+		{
+			return
+				Component->GetClass() == UDMXPixelMappingFixtureGroupComponent::StaticClass() ||
+				Component->GetClass() == UDMXPixelMappingMatrixComponent::StaticClass() ||
+				Component->GetClass() == UDMXPixelMappingScreenComponent::StaticClass() ||
+				(Component->GetParent() && Component->GetParent()->GetClass() == UDMXPixelMappingFixtureGroupComponent::StaticClass()) ||
+				(Component->GetParent() && Component->GetParent()->GetClass() == UDMXPixelMappingMatrixComponent::StaticClass()) ||
+				(Component->GetParent() && Component->GetParent()->GetClass() == UDMXPixelMappingScreenComponent::StaticClass());
+		}
+	}
+	return false;
+}
+
+void FDMXPixelMappingToolkit::SizeSelectedComponentToTexture(bool bTransacted)
+{
+	if (!ensureMsgf(CanSizeSelectedComponentToTexture(), TEXT("Trying to size selected component to texture without previously testing CanSizeSelectedComponentToTexture.")))
+	{
+		return;
+	}
+
+	if (SelectedComponents.IsEmpty())
+	{
+		return;
+	}
+
+	const UDMXPixelMappingRendererComponent* RendererComponent = GetActiveRendererComponent();
+	if (!RendererComponent)
+	{
+		return;
+	}
+
+	const FVector2D TextureSize = RendererComponent->GetSize();
+	if (TextureSize == FVector2D::ZeroVector)
+	{
+		return;
+	}
+
+	UDMXPixelMappingOutputComponent* Component = Cast<UDMXPixelMappingOutputComponent>(SelectedComponents.Array()[0].GetComponent());
+	if (!Component)
+	{
+		return;
+	}
+
+	// Apply to parent where appropriate
+	if (Component->GetClass() != UDMXPixelMappingFixtureGroupComponent::StaticClass() &&
+		Component->GetClass() != UDMXPixelMappingMatrixComponent::StaticClass() &&
+		Component->GetClass() != UDMXPixelMappingScreenComponent::StaticClass())
+	{
+		if ((Component->GetParent() && Component->GetParent()->GetClass() == UDMXPixelMappingFixtureGroupComponent::StaticClass()) ||
+			(Component->GetParent() && Component->GetParent()->GetClass() == UDMXPixelMappingMatrixComponent::StaticClass()) ||
+			(Component->GetParent() && Component->GetParent()->GetClass() == UDMXPixelMappingScreenComponent::StaticClass()))
+		{
+			Component = Cast<UDMXPixelMappingOutputComponent>(Component->GetParent());
+		}
+	}
+
+	TSharedPtr<FScopedTransaction> SizeComponentToTextureTransaction;
+	if (bTransacted)
+	{ 
+		SizeComponentToTextureTransaction = MakeShared<FScopedTransaction>(LOCTEXT("SizeComponentToTextureTransaction", "Size Component to Texture"));
+	}
+
+	const UDMXPixelMappingLayoutSettings* LayoutSettings = GetDefault<UDMXPixelMappingLayoutSettings>();
+	if (LayoutSettings && LayoutSettings->bScaleChildrenWithParent)
+	{
+		// Scale children to parent
+		const FVector2D RatioVector = TextureSize / Component->GetSize();
+		for (UDMXPixelMappingBaseComponent* BaseChild : Component->GetChildren())
+		{
+			if (UDMXPixelMappingOutputComponent* Child = Cast<UDMXPixelMappingOutputComponent>(BaseChild))
+			{
+				Child->PreEditChange(UDMXPixelMappingOutputComponent::StaticClass()->FindPropertyByName(UDMXPixelMappingOutputComponent::GetPositionXPropertyName()));
+				Child->PreEditChange(UDMXPixelMappingOutputComponent::StaticClass()->FindPropertyByName(UDMXPixelMappingOutputComponent::GetPositionYPropertyName()));
+				Child->PreEditChange(UDMXPixelMappingOutputComponent::StaticClass()->FindPropertyByName(UDMXPixelMappingOutputComponent::GetSizeXPropertyName()));
+				Child->PreEditChange(UDMXPixelMappingOutputComponent::StaticClass()->FindPropertyByName(UDMXPixelMappingOutputComponent::GetSizeYPropertyName()));
+
+				// Scale size (SetSize already clamps)
+				Child->SetSize(Child->GetSize() * RatioVector);
+
+				// Scale position (new position is zero vector)
+				const FVector2D ChildPosition = Child->GetPosition();
+				const FVector2D NewPositionRelative = (ChildPosition - Component->GetPosition()) * RatioVector;
+				Child->SetPosition(Component->GetPosition() + NewPositionRelative);
+
+				Child->PostEditChange();
+			}
+		}
+	}
+
+	Component->PreEditChange(UDMXPixelMappingOutputComponent::StaticClass()->FindPropertyByName(UDMXPixelMappingOutputComponent::GetPositionXPropertyName()));
+	Component->PreEditChange(UDMXPixelMappingOutputComponent::StaticClass()->FindPropertyByName(UDMXPixelMappingOutputComponent::GetPositionYPropertyName()));
+	Component->PreEditChange(UDMXPixelMappingOutputComponent::StaticClass()->FindPropertyByName(UDMXPixelMappingOutputComponent::GetSizeXPropertyName()));
+	Component->PreEditChange(UDMXPixelMappingOutputComponent::StaticClass()->FindPropertyByName(UDMXPixelMappingOutputComponent::GetSizeYPropertyName()));
+
+	Component->SetPosition(FVector2D::ZeroVector);
+	Component->SetSize(TextureSize);
+
+	Component->PostEditChange();
+}
+
 void FDMXPixelMappingToolkit::OnComponentAdded(UDMXPixelMapping* PixelMapping, UDMXPixelMappingBaseComponent* Component)
 {
 	// For performance reasons don't update blueprint nodes if the methods in this class add (possibly many) components.
@@ -543,10 +564,9 @@
 		FExecuteAction::CreateSP(this, &FDMXPixelMappingToolkit::DeleteSelectedComponents)
 	);
 
-	CreateInternalViewModels();
 	CreateInternalViews();
 
-	TSharedRef<FTabManager::FLayout> StandaloneDefaultLayout = FTabManager::NewLayout("Standalone_PixelMapping_Layout_v6")
+	TSharedRef<FTabManager::FLayout> StandaloneDefaultLayout = FTabManager::NewLayout("Standalone_PixelMapping_Layout_v8")
 		->AddArea
 		(
 			FTabManager::NewPrimaryArea()->SetOrientation(Orient_Vertical)
@@ -562,8 +582,14 @@
 					->Split
 					(
 						FTabManager::NewStack()
-						->AddTab(PaletteViewTabID, ETabState::OpenedTab)
-						->SetSizeCoefficient(0.5f)
+						->AddTab(InputSourceViewTabID, ETabState::OpenedTab)
+						->SetSizeCoefficient(0.25f)
 					)
 					->Split
 					(
+						FTabManager::NewStack()
+						->AddTab(DMXLibraryViewTabID, ETabState::OpenedTab)
+						->SetSizeCoefficient(0.25f)
+					)
+					->Split
+					(
@@ -572,6 +598,7 @@
 						->SetSizeCoefficient(0.5f)
 					)
 				)
+				->SetSizeCoefficient(0.25f)
 				->Split
 				(
 					FTabManager::NewSplitter()
@@ -590,6 +617,7 @@
 						->SetSizeCoefficient(0.4f)
 					)
 				)
+				->SetSizeCoefficient(0.5f)
 				->Split
 				(
 					FTabManager::NewSplitter()
@@ -608,6 +636,7 @@
 						->SetSizeCoefficient(0.4f)
 					)
 				)
+				->SetSizeCoefficient(0.25f)
 			)
 		);
 
@@ -615,20 +644,33 @@
 	const bool bCreateDefaultToolbar = true;
 	FAssetEditorToolkit::InitAssetEditor(Mode, InitToolkitHost, FDMXPixelMappingEditorModule::DMXPixelMappingEditorAppIdentifier,
 		StandaloneDefaultLayout, bCreateDefaultStandaloneMenu, bCreateDefaultToolbar, DMXPixelMapping);
-
+	
 	SetupCommands();
 	ExtendToolbar();
 	RegenerateMenusAndToolbars();
 }
 
-TSharedRef<SDockTab> FDMXPixelMappingToolkit::SpawnTab_PaletteView(const FSpawnTabArgs& Args)
+TSharedRef<SDockTab> FDMXPixelMappingToolkit::SpawnTab_InputSourceView(const FSpawnTabArgs& Args)
 {
-	check(Args.GetTabId() == PaletteViewTabID);
+	check(Args.GetTabId() == InputSourceViewTabID);
 
 	TSharedRef<SDockTab> SpawnedTab = SNew(SDockTab)
-		.Label(LOCTEXT("PaletteViewTabID", "Palette"))
+		.Label(LOCTEXT("InputSourceViewTabID", "InputSource"))
 		[
-			PaletteView.ToSharedRef()
+			InputSourceView.ToSharedRef()
 		];
 
 	return SpawnedTab;
+}
+
+TSharedRef<SDockTab> FDMXPixelMappingToolkit::SpawnTab_DMXLibraryView(const FSpawnTabArgs& Args)
+{
+	check(Args.GetTabId() == DMXLibraryViewTabID);
+
+	TSharedRef<SDockTab> SpawnedTab = SNew(SDockTab)
+		.Label(LOCTEXT("DMXLibraryViewTabID", "DMXLibrary"))
+		[
+			DMXLibraryView.ToSharedRef()
+		];
+
+	return SpawnedTab;
@@ -699,16 +741,10 @@
 	return SpawnedTab;
 }
 
-void FDMXPixelMappingToolkit::CreateInternalViewModels()
-{
-	TSharedPtr<FDMXPixelMappingToolkit> ThisPtr(SharedThis(this));
-
-	PaletteViewModel = MakeShared<FDMXPixelMappingPaletteViewModel>(ThisPtr);
-}
-
 void FDMXPixelMappingToolkit::CreateInternalViews()
 {
-	CreateOrGetView_PaletteView();
+	CreateOrGetView_InputSourceView();
+	CreateOrGetView_DMXLibraryView();
 	CreateOrGetView_HierarchyView();
 	CreateOrGetView_DesignerView();
 	CreateOrGetView_PreviewView();
@@ -721,14 +757,24 @@
 	UpdateBlueprintNodes(GetDMXPixelMapping());
 }
 
-TSharedRef<SWidget> FDMXPixelMappingToolkit::CreateOrGetView_PaletteView()
+TSharedRef<SWidget> FDMXPixelMappingToolkit::CreateOrGetView_InputSourceView()
 {
-	if (!PaletteView.IsValid())
+	if (!InputSourceView.IsValid())
 	{
-		PaletteView = SNew(SDMXPixelMappingPaletteView, SharedThis(this));
+		InputSourceView = SNew(SDMXPixelMappingInputSourceView, SharedThis(this));
 	}
 
-	return PaletteView.ToSharedRef();
+	return InputSourceView.ToSharedRef();
 }
 
+TSharedRef<SWidget> FDMXPixelMappingToolkit::CreateOrGetView_DMXLibraryView()
+{
+	if (!DMXLibraryView.IsValid())
+	{
+		DMXLibraryView = SNew(SDMXPixelMappingDMXLibraryView, SharedThis(this));
+	}
+
+	return DMXLibraryView.ToSharedRef();
+}
+
 TSharedRef<SWidget> FDMXPixelMappingToolkit::CreateOrGetView_HierarchyView()
@@ -832,9 +878,10 @@
 	);
 
 	// Layout related
+	constexpr bool bTransactSizeComponentToTexture = true;
 	GetToolkitCommands()->MapAction(
 		FDMXPixelMappingEditorCommands::Get().SizeComponentToTexture,
-		FExecuteAction::CreateSP(this, &FDMXPixelMappingToolkit::SizeSelectedComponentToTexture),
+		FExecuteAction::CreateSP(this, &FDMXPixelMappingToolkit::SizeSelectedComponentToTexture, bTransactSizeComponentToTexture),
 		FCanExecuteAction::CreateSP(this, &FDMXPixelMappingToolkit::CanSizeSelectedComponentToTexture)
 	);
 
diff --git a/Engine/Plugins/VirtualProduction/DMX/DMXPixelMapping/Source/DMXPixelMappingEditor/Private/Toolkits/DMXPixelMappingToolkit.h b/Engine/Plugins/VirtualProduction/DMX/DMXPixelMapping/Source/DMXPixelMappingEditor/Private/Toolkits/DMXPixelMappingToolkit.h
--- a/Engine/Plugins/VirtualProduction/DMX/DMXPixelMapping/Source/DMXPixelMappingEditor/Private/Toolkits/DMXPixelMappingToolkit.h
+++ b/Engine/Plugins/VirtualProduction/DMX/DMXPixelMapping/Source/DMXPixelMappingEditor/Private/Toolkits/DMXPixelMappingToolkit.h
@@ -14,10 +14,10 @@ class FDMXPixelMappingToolbar;
 class SDMXPixelMappingHierarchyView;
 class SDMXPixelMappingDesignerView;
 class SDMXPixelMappingDetailsView;
+class SDMXPixelMappingDMXLibraryView;
 class SDMXPixelMappingLayoutView;
-class SDMXPixelMappingPaletteView;
+class SDMXPixelMappingInputSourceView;
 class SDMXPixelMappingPreviewView;
-class FDMXPixelMappingPaletteViewModel;
 class UDMXPixelMapping;
 class UDMXPixelMappingBaseComponent;
 class UDMXPixelMappingMatrixComponent;
@@ -104,12 +104,13 @@ public:
 
 	const TSharedPtr<FUICommandList>& GetDesignerCommandList() const { return DesignerCommandList; }
 
-	const TSharedPtr<FDMXPixelMappingPaletteViewModel>& GetPaletteViewModel() const { return PaletteViewModel; }
-
 	const TSet<FDMXPixelMappingComponentReference>& GetSelectedComponents() const { return SelectedComponents; }
 
-	/** Gets or creates the Palette View for this Pixel Mapping instance */
-	TSharedRef<SWidget> CreateOrGetView_PaletteView();
+	/** Gets or creates the Input Source View for this Pixel Mapping instance */
+	TSharedRef<SWidget> CreateOrGetView_InputSourceView();
+
+	/** Gets or creates the DMX Library View for this Pixel Mapping instance */
+	TSharedRef<SWidget> CreateOrGetView_DMXLibraryView();
 
 	/** Gets or creates the Hierarchy View for this Pixel Mapping instance */
 	TSharedRef<SWidget> CreateOrGetView_HierarchyView();
@@ -139,6 +140,7 @@ public:
 	/** Returns true if the component is selected */
 	bool IsComponentSelected(UDMXPixelMappingBaseComponent* Component) const;
 
+	/** Adds a renderer to the pixel mapping. Sets it as active renderer if a new renderer was created */
 	void AddRenderer();
 
 	void OnComponentRenamed(UDMXPixelMappingBaseComponent* InComponent);
@@ -152,6 +154,15 @@ public:
 	/** Deletes the selected Components */
 	void DeleteSelectedComponents();
 
+	/** Zooms to fit */
+	void ZoomToFit();
+
+	/** Returns true if the selected component can be sized to texture */
+	bool CanSizeSelectedComponentToTexture() const;
+
+	/** Sizes the component to the render target of the pixelmapping asset */
+	void SizeSelectedComponentToTexture(bool bTransacted);
+
 private:
 	/** Called when a Component was added to the pixel mapping */
 	void OnComponentAdded(UDMXPixelMapping* PixelMapping, UDMXPixelMappingBaseComponent* Component);
@@ -165,20 +176,17 @@ private:
 
 	void ExecutebTogglePlayDMXAll();
 
-	/** Returns true if the selected component can be sized to texture */
-	bool CanSizeSelectedComponentToTexture() const;
-
-	/** Sizes the component to the render target of the pixelmapping asset */
-	void SizeSelectedComponentToTexture();
-
 	void UpdateBlueprintNodes(UDMXPixelMapping* InDMXPixelMapping);
 
 	void OnSaveThumbnailImage();
 
 	void InitializeInternal(const EToolkitMode::Type Mode, const TSharedPtr<class IToolkitHost>& InitToolkitHost, const FGuid& MessageLogGuid);
 
-	/** Spawns the Pallette View */
-	TSharedRef<SDockTab> SpawnTab_PaletteView(const FSpawnTabArgs& Args);
+	/** Spawns the Input Source View */
+	TSharedRef<SDockTab> SpawnTab_InputSourceView(const FSpawnTabArgs& Args);
+
+	/** Spawns the DMX Library View */
+	TSharedRef<SDockTab> SpawnTab_DMXLibraryView(const FSpawnTabArgs& Args);
 
 	/** Spawns the Hierarchy View */
 	TSharedRef<SDockTab> SpawnTab_HierarchyView(const FSpawnTabArgs& Args);
@@ -199,8 +207,6 @@ private:
 
 	void ExtendToolbar();
 
-	void CreateInternalViewModels();
-
 	void CreateInternalViews();
 
 	UDMXPixelMapping* DMXPixelMapping;
@@ -208,9 +214,12 @@ private:
 	/** List of open tool panels; used to ensure only one exists at any one time */
 	TMap<FName, TWeakPtr<SDockableTab>> SpawnedToolPanels;
 
-	/** The Palette View instance */
-	TSharedPtr<SDMXPixelMappingPaletteView> PaletteView;
+	/** The Input Source View instance */
+	TSharedPtr<SDMXPixelMappingInputSourceView> InputSourceView;
 	
+	/** The Input Source View instance */
+	TSharedPtr<SDMXPixelMappingDMXLibraryView> DMXLibraryView;
+
 	/** The Hierarchy View instance */
 	TSharedPtr<SDMXPixelMappingHierarchyView> HierarchyView;
 
@@ -226,9 +235,6 @@ private:
 	/** The Layout View instance */
 	TSharedPtr<SDMXPixelMappingLayoutView> LayoutView;
 
-	/** The Palette View Model instance */
-	TSharedPtr<FDMXPixelMappingPaletteViewModel> PaletteViewModel;
-
 	FOnSelectedComponentsChangedDelegate OnSelectedComponentsChangedDelegate;
 
 	TSet<FDMXPixelMappingComponentReference> SelectedComponents;
@@ -258,8 +264,11 @@ private:
 	bool bRemovingComponents = false;
 
 public:
-	/** Name of the Palette View Tab */
-	static const FName PaletteViewTabID;
+	/** Name of the Input Source View Tab */
+	static const FName InputSourceViewTabID;
+
+	/** Name of the DMX Library View Tab */
+	static const FName DMXLibraryViewTabID;
 
 	/** Name of the Hierarchy View Tab */
 	static const FName HierarchyViewTabID;
diff --git a/Engine/Plugins/VirtualProduction/DMX/DMXPixelMapping/Source/DMXPixelMappingEditor/Private/ViewModels/DMXPixelMappingDMXLibraryViewModel.cpp b/Engine/Plugins/VirtualProduction/DMX/DMXPixelMapping/Source/DMXPixelMappingEditor/Private/ViewModels/DMXPixelMappingDMXLibraryViewModel.cpp
new file mode 100644
--- /dev/null
+++ b/Engine/Plugins/VirtualProduction/DMX/DMXPixelMapping/Source/DMXPixelMappingEditor/Private/ViewModels/DMXPixelMappingDMXLibraryViewModel.cpp
@@ -0,0 +1,194 @@
+// Copyright Epic Games, Inc. All Rights Reserved.
+
+#include "DMXPixelMappingDMXLibraryViewModel.h"
+
+#include "DMXPixelMapping.h"
+#include "Components/DMXPixelMappingFixtureGroupComponent.h"
+#include "Components/DMXPixelMappingFixtureGroupItemComponent.h"
+#include "Components/DMXPixelMappingMatrixComponent.h"
+#include "Components/DMXPixelMappingRendererComponent.h"
+#include "Library/DMXEntityFixturePatch.h"
+#include "Toolkits/DMXPixelMappingToolkit.h"
+
+#include "Editor.h"
+
+
+UDMXPixelMappingDMXLibraryViewModel* UDMXPixelMappingDMXLibraryViewModel::CreateNew(const TSharedPtr<FDMXPixelMappingToolkit>& InToolkit, UDMXPixelMappingFixtureGroupComponent* InFixtureGroup)
+{
+	UDMXPixelMappingDMXLibraryViewModel* NewModel = NewObject<UDMXPixelMappingDMXLibraryViewModel>(GetTransientPackage(), NAME_None, RF_Transactional);
+	NewModel->WeakToolkit = InToolkit;
+	NewModel->WeakFixtureGroupComponent = InFixtureGroup;
+	NewModel->DMXLibrary = InFixtureGroup ? InFixtureGroup->DMXLibrary : nullptr;
+
+	if (NewModel->WeakFixtureGroupComponent.IsValid())
+	{
+		NewModel->WeakFixtureGroupComponent->GetOnDMXLibraryChanged().AddUObject(NewModel, &UDMXPixelMappingDMXLibraryViewModel::UpdateDMXLibraryFromComponent);
+	}
+
+	return NewModel;
+}
+
+UDMXPixelMappingFixtureGroupComponent* UDMXPixelMappingDMXLibraryViewModel::GetFixtureGroupComponent() const
+{
+	return WeakFixtureGroupComponent.Get();
+}
+
+UDMXPixelMappingRootComponent* UDMXPixelMappingDMXLibraryViewModel::GetPixelMappingRootComponent() const
+{
+	const TSharedPtr<FDMXPixelMappingToolkit> Toolkit = WeakToolkit.IsValid() ? WeakToolkit.Pin() : nullptr;
+	UDMXPixelMapping* PixelMapping = Toolkit.IsValid() ? Toolkit->GetDMXPixelMapping() : nullptr;
+
+	return PixelMapping ? PixelMapping->GetRootComponent() : nullptr;
+}
+
+bool UDMXPixelMappingDMXLibraryViewModel::IsFixtureGroupOrChildSelected() const
+{
+	if (!WeakFixtureGroupComponent.IsValid() || !WeakToolkit.IsValid())
+	{
+		return false;
+	}
+
+	const TSet<FDMXPixelMappingComponentReference>& SelectedComponents = WeakToolkit.Pin()->GetSelectedComponents();
+	for (const FDMXPixelMappingComponentReference& ComponentReference : SelectedComponents)
+	{
+		UDMXPixelMappingBaseComponent* Component = ComponentReference.GetComponent();
+		do
+		{
+			if (Component == WeakFixtureGroupComponent)
+			{
+				return true;
+			}
+
+			Component = Component->GetParent();
+		} while (Component);
+	}
+
+	return false;
+}
+
+const TArray<FDMXEntityFixturePatchRef> UDMXPixelMappingDMXLibraryViewModel::GetFixturePatchesInUse() const
+{
+	TArray<FDMXEntityFixturePatchRef> FixturePatchesInUse;
+
+	const TArray<UDMXPixelMappingFixtureGroupComponent*> FixtureGroupComponents = GetFixtureGroupComponentsOfSameLibrary();
+	for (UDMXPixelMappingFixtureGroupComponent* FixtureGroupComponent : FixtureGroupComponents)
+	{
+		for (UDMXPixelMappingBaseComponent* Child : FixtureGroupComponent->GetChildren())
+		{
+			if (UDMXPixelMappingFixtureGroupItemComponent* GroupItem = Cast<UDMXPixelMappingFixtureGroupItemComponent>(Child))
+			{
+				if (GroupItem->FixturePatchRef.DMXLibrary == DMXLibrary)
+				{
+					FixturePatchesInUse.Add(GroupItem->FixturePatchRef);
+				}
+			}
+			else if (UDMXPixelMappingMatrixComponent* MatrixComponent = Cast<UDMXPixelMappingMatrixComponent>(Child))
+			{
+				if (MatrixComponent->FixturePatchRef.DMXLibrary == DMXLibrary)
+				{
+					FixturePatchesInUse.Add(MatrixComponent->FixturePatchRef);
+				}
+			}
+		}
+	}
+
+	return FixturePatchesInUse;
+}
+
+void UDMXPixelMappingDMXLibraryViewModel::SaveFixturePatchListDescriptor(const FDMXReadOnlyFixturePatchListDescriptor& NewDescriptor)
+{
+	FixturePatchListDescriptor = NewDescriptor;
+	SaveConfig();
+}
+
+void UDMXPixelMappingDMXLibraryViewModel::PostUndo(bool bSuccess)
+{
+	UpdateDMXLibraryFromComponent();
+}
+
+void UDMXPixelMappingDMXLibraryViewModel::PostRedo(bool bSuccess)
+{
+	UpdateDMXLibraryFromComponent();
+}
+
+void UDMXPixelMappingDMXLibraryViewModel::PostEditChangeProperty(FPropertyChangedEvent& PropertyChangedEvent)
+{
+	Super::PostEditChangeProperty(PropertyChangedEvent);
+
+	TGuardValue<bool> ChangingPropertiesGuard(bChangingProperties, true);
+
+	if (PropertyChangedEvent.GetPropertyName() == GET_MEMBER_NAME_CHECKED(UDMXPixelMappingDMXLibraryViewModel, DMXLibrary))
+	{
+		if (WeakFixtureGroupComponent.IsValid())
+		{
+			WeakFixtureGroupComponent->PreEditChange(UDMXPixelMappingFixtureGroupComponent::StaticClass()->FindPropertyByName(GET_MEMBER_NAME_CHECKED(UDMXPixelMappingFixtureGroupComponent, DMXLibrary)));
+			WeakFixtureGroupComponent->DMXLibrary = DMXLibrary;
+			WeakFixtureGroupComponent->PostEditChange();
+
+			RemoveInvalidPatches();
+
+			OnDMXLibraryChanged.Broadcast();
+		}
+	}
+}
+
+void UDMXPixelMappingDMXLibraryViewModel::UpdateDMXLibraryFromComponent()
+{
+	if (bChangingProperties)
+	{
+		return;
+	}
+
+	UDMXPixelMappingFixtureGroupComponent* FixtureGroupComponent = WeakFixtureGroupComponent.Get();
+	if (FixtureGroupComponent)
+	{
+		Modify();
+		DMXLibrary = FixtureGroupComponent->DMXLibrary;
+		RemoveInvalidPatches();
+	}
+}
+
+void UDMXPixelMappingDMXLibraryViewModel::RemoveInvalidPatches()
+{
+	if (UDMXPixelMappingFixtureGroupComponent* FixtureGroupComponent = WeakFixtureGroupComponent.Get())
+	{
+		TArray<UDMXPixelMappingBaseComponent*> CachedChildren(WeakFixtureGroupComponent->Children);
+		for (UDMXPixelMappingBaseComponent* ChildComponent : CachedChildren)
+		{
+			if (UDMXPixelMappingFixtureGroupItemComponent* ChildGroupItem = Cast<UDMXPixelMappingFixtureGroupItemComponent>(ChildComponent))
+			{
+				if (ChildGroupItem->FixturePatchRef.GetFixturePatch() &&
+					ChildGroupItem->FixturePatchRef.GetFixturePatch()->GetParentLibrary() != FixtureGroupComponent->DMXLibrary)
+				{
+					FixtureGroupComponent->RemoveChild(ChildGroupItem);
+				}
+			}
+			else if (UDMXPixelMappingMatrixComponent* ChildMatrix = Cast<UDMXPixelMappingMatrixComponent>(ChildComponent))
+			{
+				if (ChildMatrix->FixturePatchRef.GetFixturePatch() &&
+					ChildMatrix->FixturePatchRef.GetFixturePatch()->GetParentLibrary() != FixtureGroupComponent->DMXLibrary)
+				{
+					FixtureGroupComponent->RemoveChild(ChildMatrix);
+				}
+			}
+		};
+	}
+}
+
+TArray<UDMXPixelMappingFixtureGroupComponent*> UDMXPixelMappingDMXLibraryViewModel::GetFixtureGroupComponentsOfSameLibrary() const
+{
+	TArray<UDMXPixelMappingFixtureGroupComponent*> FixtureGroupComponents;
+	UDMXPixelMappingRendererComponent* RendererComponent = WeakToolkit.IsValid() ? WeakToolkit.Pin()->GetActiveRendererComponent() : nullptr;
+	if (RendererComponent)
+	{
+		for (UDMXPixelMappingBaseComponent* Child : RendererComponent->GetChildren())
+		{
+			if (UDMXPixelMappingFixtureGroupComponent* GroupComponent = Cast<UDMXPixelMappingFixtureGroupComponent>(Child))
+			{
+				FixtureGroupComponents.Add(GroupComponent);
+			}
+		}
+	}
+
+	return FixtureGroupComponents;
+}
diff --git a/Engine/Plugins/VirtualProduction/DMX/DMXPixelMapping/Source/DMXPixelMappingEditor/Private/ViewModels/DMXPixelMappingDMXLibraryViewModel.h b/Engine/Plugins/VirtualProduction/DMX/DMXPixelMapping/Source/DMXPixelMappingEditor/Private/ViewModels/DMXPixelMappingDMXLibraryViewModel.h
new file mode 100644
--- /dev/null
+++ b/Engine/Plugins/VirtualProduction/DMX/DMXPixelMapping/Source/DMXPixelMappingEditor/Private/ViewModels/DMXPixelMappingDMXLibraryViewModel.h
@@ -0,0 +1,95 @@
+// Copyright Epic Games, Inc. All Rights Reserved.
+
+#pragma once
+
+#include "Library/DMXEntityReference.h"
+#include "Widgets/SDMXReadOnlyFixturePatchList.h"
+
+#include "EditorUndoClient.h"
+#include "Delegates/Delegate.h"
+#include "UObject/Object.h"
+#include "UObject/WeakObjectPtr.h"
+
+#include "DMXPixelMappingDMXLibraryViewModel.generated.h"
+
+class FDMXPixelMappingToolkit;
+class UDMXEntityFixturePatch;
+class UDMXLibrary;
+class UDMXPixelMappingFixtureGroupComponent;
+class UDMXPixelMappingRootComponent;
+
+
+/** Model for the DMX Library View */
+UCLASS(Config = DMXEditor)
+class UDMXPixelMappingDMXLibraryViewModel
+	: public UObject
+	, public FEditorUndoClient
+{
+	GENERATED_BODY()
+
+public:
+	/** Creates a new model */
+	static UDMXPixelMappingDMXLibraryViewModel* CreateNew(const TSharedPtr<FDMXPixelMappingToolkit>& InToolkit, UDMXPixelMappingFixtureGroupComponent* InFixtureGroup);
+
+	/** Returns the DMX Library */
+	UDMXLibrary* GetDMXLibrary() const { return DMXLibrary;  }
+
+	/** Sets the fixture group this model uses */ 
+	UDMXPixelMappingFixtureGroupComponent* GetFixtureGroupComponent() const;
+
+	/** Returns the root component of the pixel mapping */
+	UDMXPixelMappingRootComponent* GetPixelMappingRootComponent() const;
+
+	/** Returns true if the fixture group or one of its children is selected */
+	bool IsFixtureGroupOrChildSelected() const;
+
+	/** Returns the fixture patches in use in the pixel mapping */
+	const TArray<FDMXEntityFixturePatchRef> GetFixturePatchesInUse() const;
+
+	/** Returns the default fixture patch list descriptor */
+	const FDMXReadOnlyFixturePatchListDescriptor& GetFixturePatchListDescriptor() const { return FixturePatchListDescriptor; }
+
+	/** Saves the fixture patch list descriptor */
+	void SaveFixturePatchListDescriptor(const FDMXReadOnlyFixturePatchListDescriptor& NewDescriptor);
+
+	/** Delegate broadcast when the DMX library changed */
+	FSimpleMulticastDelegate OnDMXLibraryChanged;
+
+protected:
+	//~ Begin FEditorUndoClient Interface
+	virtual void PostUndo(bool bSuccess) override;
+	virtual void PostRedo(bool bSuccess) override;
+	// End of FEditorUndoClient
+
+	// Begin UObject interface
+	virtual void PostEditChangeProperty(FPropertyChangedEvent& PropertyChangedEvent) override;
+	// End UObject interface
+
+private:
+	/** Updates the DMX Library from the component */
+	void UpdateDMXLibraryFromComponent();
+
+	/** Removes invalid patches from pixel mapping */
+	void RemoveInvalidPatches();
+
+	/** Returns the fixture group components in the pixel mapping */
+	TArray<UDMXPixelMappingFixtureGroupComponent*> GetFixtureGroupComponentsOfSameLibrary() const;
+
+	/** The DMX library of this view */
+	UPROPERTY(EditAnywhere, Transient, Category = "DMXLibrary", Meta = (AllowPrivateAccess = true))
+	TObjectPtr<UDMXLibrary> DMXLibrary;
+
+	/** The fixture group component that this model uses */
+	UPROPERTY(Transient)
+	TWeakObjectPtr<UDMXPixelMappingFixtureGroupComponent> WeakFixtureGroupComponent;
+
+	/** Default fixture patch list descriptor */
+	UPROPERTY(Config)
+	FDMXReadOnlyFixturePatchListDescriptor FixturePatchListDescriptor;
+
+	/** True while changing properties */
+	bool bChangingProperties = false;
+
+	/** The toolkit of the editor that uses this model */
+	TWeakPtr<FDMXPixelMappingToolkit> WeakToolkit;
+};
diff --git a/Engine/Plugins/VirtualProduction/DMX/DMXPixelMapping/Source/DMXPixelMappingEditor/Private/ViewModels/DMXPixelMappingPaletteViewModel.cpp b/Engine/Plugins/VirtualProduction/DMX/DMXPixelMapping/Source/DMXPixelMappingEditor/Private/ViewModels/DMXPixelMappingPaletteViewModel.cpp
deleted file mode 100644
diff --git a/Engine/Plugins/VirtualProduction/DMX/DMXPixelMapping/Source/DMXPixelMappingEditor/Private/ViewModels/DMXPixelMappingPaletteViewModel.h b/Engine/Plugins/VirtualProduction/DMX/DMXPixelMapping/Source/DMXPixelMappingEditor/Private/ViewModels/DMXPixelMappingPaletteViewModel.h
deleted file mode 100644
diff --git a/Engine/Plugins/VirtualProduction/DMX/DMXPixelMapping/Source/DMXPixelMappingEditor/Private/Views/SDMXPixelMappingDMXLibraryView.cpp b/Engine/Plugins/VirtualProduction/DMX/DMXPixelMapping/Source/DMXPixelMappingEditor/Private/Views/SDMXPixelMappingDMXLibraryView.cpp
new file mode 100644
--- /dev/null
+++ b/Engine/Plugins/VirtualProduction/DMX/DMXPixelMapping/Source/DMXPixelMappingEditor/Private/Views/SDMXPixelMappingDMXLibraryView.cpp
@@ -0,0 +1,332 @@
+// Copyright Epic Games, Inc. All Rights Reserved.
+
+#include "Views/SDMXPixelMappingDMXLibraryView.h"
+
+#include "DMXPixelMapping.h"
+#include "DMXPixelMappingComponentReference.h"
+#include "Components/DMXPixelMappingFixtureGroupComponent.h"
+#include "Components/DMXPixelMappingFixtureGroupItemComponent.h"
+#include "Components/DMXPixelMappingMatrixComponent.h"
+#include "Components/DMXPixelMappingRendererComponent.h"
+#include "Components/DMXPixelMappingRootComponent.h"
+#include "Customizations/DMXPixelMappingDetailCustomization_FixtureGroup.h"
+#include "Library/DMXLibrary.h"
+#include "Templates/DMXPixelMappingComponentTemplate.h"
+#include "Toolkits/DMXPixelMappingToolkit.h"
+#include "ViewModels/DMXPixelMappingDMXLibraryViewModel.h"
+#include "Widgets/SDMXPixelMappingFixturePatchList.h"
+
+#include "Editor.h"
+#include "IDetailsView.h"
+#include "PropertyEditorDelegates.h"
+#include "PropertyEditorModule.h"
+#include "ScopedTransaction.h"
+#include "Styling/AppStyle.h"
+#include "Widgets/SBoxPanel.h"
+#include "Widgets/Images/SImage.h"
+#include "Widgets/Input/SButton.h"
+#include "Widgets/Layout/SExpandableArea.h"
+#include "Widgets/Layout/SScrollBox.h"
+
+
+#define LOCTEXT_NAMESPACE "SDMXPixelMappingDMXLibraryView"
+
+void SDMXPixelMappingDMXLibraryView::Construct(const FArguments& InArgs, const TSharedPtr<FDMXPixelMappingToolkit>& InToolkit)
+{
+	WeakToolkit = InToolkit;
+	if (!WeakToolkit.IsValid())
+	{
+		return;
+	}
+	const TSharedRef<FDMXPixelMappingToolkit> Toolkit = WeakToolkit.Pin().ToSharedRef();
+
+	ChildSlot
+	[
+		SNew(SBorder)
+		.BorderImage(FAppStyle::GetBrush("ToolPanel.GroupBorder"))
+		[
+			SNew(SVerticalBox)
+
+			// Add button
+			+ SVerticalBox::Slot()
+			.HAlign(HAlign_Left)
+			.Padding(4.f)
+			.AutoHeight()
+			[
+				SNew(SButton)
+				.ButtonStyle(FAppStyle::Get(), "FlatButton.Success")
+				.ForegroundColor(FLinearColor::White)
+				.ToolTipText(LOCTEXT("AddFixtureGroupTooltip", "Adds a Fixture Group to the Pixel Mapping"))
+				.ContentPadding(FMargin(5.0f, 1.0f))
+				.OnClicked(this, &SDMXPixelMappingDMXLibraryView::OnAddDMXLibraryButtonClicked)
+				.Content()
+				[
+					SNew(SHorizontalBox)
+
+					+ SHorizontalBox::Slot()
+					.AutoWidth()
+					.Padding(FMargin(0.f, 1.f))
+					[
+						SNew(SImage)
+						.Image(FAppStyle::GetBrush("Plus"))
+					]
+
+					+ SHorizontalBox::Slot()
+					.VAlign(VAlign_Center)
+					.AutoWidth()
+					.Padding(FMargin(2.f, 0.f, 2.f, 0.f))
+					[
+						SNew(STextBlock)
+						.Text(LOCTEXT("AddFixtureGroupLabel", "Add DMX Library"))
+					]
+				]
+			]
+
+			// Details views 
+			+ SVerticalBox::Slot()
+			.FillHeight(1.f)
+			[
+				SAssignNew(DMXLibrariesScrollBox, SScrollBox)
+				.Orientation(EOrientation::Orient_Vertical)
+			]
+		]
+	];
+
+	ForceRefresh();
+
+	Toolkit->GetOnSelectedComponentsChangedDelegate().AddSP(this, &SDMXPixelMappingDMXLibraryView::OnComponentSelected);
+}
+
+void SDMXPixelMappingDMXLibraryView::RequestRefresh()
+{
+	if (!RefreshTimerHandle.IsValid())
+	{
+		RefreshTimerHandle = GEditor->GetTimerManager()->SetTimerForNextTick(FTimerDelegate::CreateSP(this, &SDMXPixelMappingDMXLibraryView::ForceRefresh));
+	}
+}
+
+void SDMXPixelMappingDMXLibraryView::AddReferencedObjects(FReferenceCollector& Collector) 
+{
+	Collector.AddReferencedObjects(ViewModels);
+}
+
+void SDMXPixelMappingDMXLibraryView::ForceRefresh()
+{
+	RefreshTimerHandle.Invalidate();
+
+	if (!WeakToolkit.IsValid())
+	{
+		return;
+	}
+	const TSharedRef<FDMXPixelMappingToolkit> Toolkit = WeakToolkit.Pin().ToSharedRef();
+
+	DMXLibrariesScrollBox->ClearChildren();
+	ViewModels.Reset();
+
+	const TArray<UDMXPixelMappingFixtureGroupComponent*> FixtureGroupComponents = GetFixtureGroupComponents();
+	for (UDMXPixelMappingFixtureGroupComponent* FixtureGroupComponent : FixtureGroupComponents)
+	{
+		UDMXPixelMappingDMXLibraryViewModel* ViewModel = UDMXPixelMappingDMXLibraryViewModel::CreateNew(Toolkit, FixtureGroupComponent);
+		if (!ViewModel->IsFixtureGroupOrChildSelected())
+		{
+			// Only display selected models
+			continue;
+		}
+		ViewModels.Add(ViewModel);
+
+		// Listen to dmx library changes
+		ViewModel->OnDMXLibraryChanged.AddSP(this, &SDMXPixelMappingDMXLibraryView::RequestRefresh);
+
+		// Create title content
+		const TSharedRef<SHorizontalBox> TitleContentBox = SNew(SHorizontalBox);
+		TitleContentBox->AddSlot()
+			.HAlign(HAlign_Left)
+			.VAlign(VAlign_Center)
+			.FillWidth(1.f)
+			[
+				SNew(STextBlock)
+				.TextStyle(FAppStyle::Get(), "DetailsView.CategoryTextStyle")
+				.Text_Lambda([FixtureGroupComponent]()
+					{
+						return FixtureGroupComponent ? FText::FromString(FixtureGroupComponent->GetUserFriendlyName()) : FText::GetEmpty();
+					})
+			];
+
+		// Create body content
+		const TSharedRef<SVerticalBox> BodyContent =
+			SNew(SVerticalBox)
+
+			+ SVerticalBox::Slot()
+			.AutoHeight()
+			[
+				CreateDetailsViewForModel(*ViewModel)
+			]
+
+			+ SVerticalBox::Slot()
+			.AutoHeight()
+			[
+				SNew(SDMXPixelMappingFixturePatchList, Toolkit, ViewModel)
+			];
+
+		// Add to scrollbox
+		DMXLibrariesScrollBox->AddSlot()
+			[
+				SNew(SExpandableArea)
+				.BorderImage(FAppStyle::Get().GetBrush("DetailsView.CategoryTop"))
+				.BodyBorderImage(FAppStyle::Get().GetBrush("Brushes.Recessed"))				
+				.HeaderPadding(4.f)
+				.Padding(FMargin(8.f, 2.f, 2.f, 0.f))
+				.HeaderContent()
+				[
+					TitleContentBox
+				]
+				.BodyContent()
+				[
+					BodyContent
+				]
+			];
+	}
+}
+
+void SDMXPixelMappingDMXLibraryView::OnComponentSelected()
+{
+	const TSharedPtr<FDMXPixelMappingToolkit> Toolkit = WeakToolkit.Pin();
+	if (!Toolkit.IsValid())
+	{
+		return;
+	}
+
+	// Refresh the view if there is no view model for the fixture group of a selected component
+	const TSet<FDMXPixelMappingComponentReference>& SelectedComponents = Toolkit->GetSelectedComponents();
+
+	UDMXPixelMappingFixtureGroupComponent* ComponentToFocus = nullptr;
+	for (const FDMXPixelMappingComponentReference& ComponentReference : SelectedComponents)
+	{
+		UDMXPixelMappingBaseComponent* Component = ComponentReference.GetComponent();
+		do
+		{
+			UDMXPixelMappingDMXLibraryViewModel* const* CorrespondingViewModelPtr = Algo::FindByPredicate(ViewModels, [Component](const UDMXPixelMappingDMXLibraryViewModel* ViewModel)
+				{
+					return ViewModel->GetFixtureGroupComponent() == Component;
+				});
+			if (!CorrespondingViewModelPtr)
+			{
+				RequestRefresh();
+				return;
+			}
+
+			Component = Component->GetParent();
+		} while (Component);
+	}
+}
+
+FReply SDMXPixelMappingDMXLibraryView::OnAddDMXLibraryButtonClicked()
+{
+	if (!WeakToolkit.IsValid())
+	{
+		return FReply::Unhandled();
+	}
+	const TSharedRef<FDMXPixelMappingToolkit> Toolkit = WeakToolkit.Pin().ToSharedRef();
+
+	UDMXPixelMappingRootComponent* RootComponent = GetPixelMappingRootComponent();
+	UDMXPixelMappingRendererComponent* ActiveRendererComponent = Toolkit->GetActiveRendererComponent();
+	if (!RootComponent || !ActiveRendererComponent)
+	{
+		return FReply::Handled();
+	}
+
+	const TArray<UDMXPixelMappingFixtureGroupComponent*> OtherFixtureGroupComponents = GetFixtureGroupComponents();
+
+	const FScopedTransaction AddFixtureGroupTransaction(LOCTEXT("AddFixtureGroupTransaction", "Add Fixture Group"));
+	const TSharedRef<FDMXPixelMappingComponentTemplate> Template = MakeShared<FDMXPixelMappingComponentTemplate>(UDMXPixelMappingFixtureGroupComponent::StaticClass());
+	const TArray<UDMXPixelMappingBaseComponent*> NewComponents = Toolkit->CreateComponentsFromTemplates(RootComponent, ActiveRendererComponent, TArray<TSharedPtr<FDMXPixelMappingComponentTemplate>>{ Template });
+	if (!ensureMsgf(NewComponents.Num() == 1 && NewComponents[0]->GetClass() == UDMXPixelMappingFixtureGroupComponent::StaticClass(), TEXT("Cannot find newly added Fixture Group Component, hence cannot select it.")))
+	{
+		return FReply::Handled();
+	}
+	UDMXPixelMappingFixtureGroupComponent* NewFixtureGroupComponent = CastChecked<UDMXPixelMappingFixtureGroupComponent>(NewComponents[0]);
+	SelectFixtureGroupComponent(NewFixtureGroupComponent);
+
+	// Place the new fixture group
+	if (OtherFixtureGroupComponents.IsEmpty())
+	{
+		// If there's no group, add one that scales the current texture to the active renderer component
+		if (Toolkit->CanSizeSelectedComponentToTexture(), TEXT("Unexpected cannot size single fixture group component to texture"))
+		{
+			Toolkit->SizeSelectedComponentToTexture(true);
+		}
+	}
+	else
+	{
+		// If there's already a group, place it over the bottom right most existing group
+		FVector2D NewPosition(0.f, 0.f);
+		FVector2D NewSize = NewFixtureGroupComponent->GetSize();
+		for (UDMXPixelMappingFixtureGroupComponent* Other : OtherFixtureGroupComponents)
+		{
+			if (Other->GetPosition().Y > NewPosition.Y)
+			{
+				NewPosition = Other->GetPosition();
+				NewSize = Other->GetSize();
+			}
+		}
+		NewPosition += FVector2D(FMath::Max(1.f, NewSize.X / 16), FMath::Max(1.f, NewSize.Y / 16));
+
+		NewFixtureGroupComponent->SetPosition(NewPosition);
+		NewFixtureGroupComponent->SetSize(NewSize);
+	}
+
+	return FReply::Handled();
+}
+
+TSharedRef<SWidget> SDMXPixelMappingDMXLibraryView::CreateDetailsViewForModel(UDMXPixelMappingDMXLibraryViewModel& Model) const
+{
+	FDetailsViewArgs DetailsViewArgs;
+	DetailsViewArgs.NameAreaSettings = FDetailsViewArgs::HideNameArea;
+	DetailsViewArgs.bHideSelectionTip = true;
+	DetailsViewArgs.DefaultsOnlyVisibility = EEditDefaultsOnlyNodeVisibility::Automatic;
+
+	FPropertyEditorModule& PropertyEditorModule = FModuleManager::Get().GetModuleChecked<FPropertyEditorModule>("PropertyEditor");
+	const TSharedRef<IDetailsView> DetailsView = PropertyEditorModule.CreateDetailView(DetailsViewArgs);
+
+	DetailsView->SetObject(&Model);
+
+	return DetailsView;
+}
+
+void SDMXPixelMappingDMXLibraryView::SelectFixtureGroupComponent(UDMXPixelMappingFixtureGroupComponent* FixtureGroupComponent)
+{
+	if (const TSharedPtr<FDMXPixelMappingToolkit> Toolkit = WeakToolkit.Pin())
+	{
+		const FDMXPixelMappingComponentReference ComponentRefToSelect(Toolkit, FixtureGroupComponent);
+		const TSet<FDMXPixelMappingComponentReference> NewSelection{ ComponentRefToSelect };
+		Toolkit->SelectComponents(NewSelection);
+	}
+}
+
+UDMXPixelMappingRootComponent* SDMXPixelMappingDMXLibraryView::GetPixelMappingRootComponent() const
+{
+	const TSharedPtr<FDMXPixelMappingToolkit> Toolkit = WeakToolkit.IsValid() ? WeakToolkit.Pin() : nullptr;
+	UDMXPixelMapping* PixelMapping = Toolkit.IsValid() ? Toolkit->GetDMXPixelMapping() : nullptr;
+
+	return PixelMapping ? PixelMapping->GetRootComponent() : nullptr;
+}
+
+TArray<UDMXPixelMappingFixtureGroupComponent*> SDMXPixelMappingDMXLibraryView::GetFixtureGroupComponents() const
+{
+	TArray<UDMXPixelMappingFixtureGroupComponent*> FixtureGroupComponents;
+	UDMXPixelMappingRendererComponent* RendererComponent = WeakToolkit.IsValid() ? WeakToolkit.Pin()->GetActiveRendererComponent() : nullptr;
+	if (RendererComponent)
+	{
+		for (UDMXPixelMappingBaseComponent* Child : RendererComponent->GetChildren())
+		{
+			if (UDMXPixelMappingFixtureGroupComponent* GroupComponent = Cast<UDMXPixelMappingFixtureGroupComponent>(Child))
+			{
+				FixtureGroupComponents.Add(GroupComponent);
+			}
+		}
+	}
+
+	return FixtureGroupComponents;
+}
+
+#undef LOCTEXT_NAMESPACE
diff --git a/Engine/Plugins/VirtualProduction/DMX/DMXPixelMapping/Source/DMXPixelMappingEditor/Private/Views/SDMXPixelMappingDMXLibraryView.h b/Engine/Plugins/VirtualProduction/DMX/DMXPixelMapping/Source/DMXPixelMappingEditor/Private/Views/SDMXPixelMappingDMXLibraryView.h
new file mode 100644
--- /dev/null
+++ b/Engine/Plugins/VirtualProduction/DMX/DMXPixelMapping/Source/DMXPixelMappingEditor/Private/Views/SDMXPixelMappingDMXLibraryView.h
@@ -0,0 +1,82 @@
+// Copyright Epic Games, Inc. All Rights Reserved.
+
+#pragma once
+
+#include "Library/DMXEntityReference.h"
+
+#include "Engine/EngineTypes.h"
+#include "Widgets/SCompoundWidget.h"
+#include "UObject/GCObject.h"
+
+class FDMXPixelMappingToolkit;
+class SDMXReadOnlyFixturePatchList;
+class UDMXPixelMappingFixtureGroupComponent;
+class UDMXPixelMappingDMXLibraryViewModel;
+class UDMXPixelMappingRootComponent;
+
+class FReply;
+class IDetailsView;
+class SScrollBox;
+
+
+/** Displays the DMX Library of the currently selected fixture group component */
+class SDMXPixelMappingDMXLibraryView
+	: public SCompoundWidget
+	, public FGCObject
+{
+public:
+	SLATE_BEGIN_ARGS(SDMXPixelMappingDMXLibraryView) { }
+	SLATE_END_ARGS()
+
+	/** Constructs this widget */
+	void Construct(const FArguments& InArgs, const TSharedPtr<FDMXPixelMappingToolkit>& InToolkit);
+
+	/** Refreshes the widget on the next tick */
+	void RequestRefresh();
+
+protected:
+	//~ Begin FGCObject interface
+	virtual void AddReferencedObjects(FReferenceCollector& Collector) override;
+	virtual FString GetReferencerName() const override
+	{
+		return TEXT("SDMXPixelMappingDMXLibraryView");
+	}
+	//~ End FGCObject interface
+
+private:
+	/** Refreshes the widget */
+	void ForceRefresh();
+
+	/** Called when a component was selected */
+	void OnComponentSelected();
+
+	/** Called when the 'Add DMX Library' button was clicked */
+	FReply OnAddDMXLibraryButtonClicked();
+
+	/** Creates a details view for the DMX Library View Model */
+	TSharedRef<SWidget> CreateDetailsViewForModel(UDMXPixelMappingDMXLibraryViewModel& Model) const;
+
+	/** Selects the fixture group component */
+	void SelectFixtureGroupComponent(UDMXPixelMappingFixtureGroupComponent* FixtureGroupComponent);
+
+	/** Returns the root component of the pixel mapping */
+	UDMXPixelMappingRootComponent* GetPixelMappingRootComponent() const;
+
+	/** Returns the fixture group components in the pixel mapping */
+	TArray<UDMXPixelMappingFixtureGroupComponent*> GetFixtureGroupComponents() const;
+
+	/** List of fixture patches the user can select from */
+	TSharedPtr<SDMXReadOnlyFixturePatchList> FixturePatchList;
+
+	/** Scrollbox that holds the details views */
+	TSharedPtr<SScrollBox> DMXLibrariesScrollBox;
+
+	/** Timer handle for the Request Refresh method */
+	FTimerHandle RefreshTimerHandle;
+
+	/** The view model of the focused DMX Library */
+	TArray<UDMXPixelMappingDMXLibraryViewModel*> ViewModels;
+
+	/** The toolkit of this editor */
+	TWeakPtr<FDMXPixelMappingToolkit> WeakToolkit;
+};
diff --git a/Engine/Plugins/VirtualProduction/DMX/DMXPixelMapping/Source/DMXPixelMappingEditor/Private/Views/SDMXPixelMappingDetailsView.cpp b/Engine/Plugins/VirtualProduction/DMX/DMXPixelMapping/Source/DMXPixelMappingEditor/Private/Views/SDMXPixelMappingDetailsView.cpp
--- a/Engine/Plugins/VirtualProduction/DMX/DMXPixelMapping/Source/DMXPixelMappingEditor/Private/Views/SDMXPixelMappingDetailsView.cpp
+++ b/Engine/Plugins/VirtualProduction/DMX/DMXPixelMapping/Source/DMXPixelMappingEditor/Private/Views/SDMXPixelMappingDetailsView.cpp
@@ -89,7 +89,11 @@ void SDMXPixelMappingDetailsView::OnSelectedComponentsChanged()
 	{
 		for (const FDMXPixelMappingComponentReference& ComponentRef : SelectedComponents)
 		{
-			SelectedObjects.Add(ComponentRef.GetComponent());
+			UDMXPixelMappingBaseComponent* Component = ComponentRef.GetComponent();
+			if (Component && Component->GetClass() != UDMXPixelMappingRendererComponent::StaticClass())
+			{
+				SelectedObjects.Add(ComponentRef.GetComponent());
+			}
 		}
 	}
 
@@ -120,13 +124,10 @@ void SDMXPixelMappingDetailsView::RegisterCustomizations()
 
 	FOnGetDetailCustomizationInstance FixtureGroupItemCustomizationInstance = FOnGetDetailCustomizationInstance::CreateStatic(&FDMXPixelMappingDetailCustomization_FixtureGroupItem::MakeInstance, ToolkitWeakPtr);
 	PropertyView->RegisterInstancedCustomPropertyLayout(UDMXPixelMappingFixtureGroupItemComponent::StaticClass(), FixtureGroupItemCustomizationInstance);
-
+	
 	FOnGetDetailCustomizationInstance ScreenCustomizationInstance = FOnGetDetailCustomizationInstance::CreateStatic(&FDMXPixelMappingDetailCustomization_Screen::MakeInstance, ToolkitWeakPtr);
 	PropertyView->RegisterInstancedCustomPropertyLayout(UDMXPixelMappingScreenComponent::StaticClass(), ScreenCustomizationInstance);
 
-	FOnGetDetailCustomizationInstance RendererCustomizationInstance = FOnGetDetailCustomizationInstance::CreateStatic(&FDMXPixelMappingDetailCustomization_Renderer::MakeInstance, ToolkitWeakPtr);
-	PropertyView->RegisterInstancedCustomPropertyLayout(UDMXPixelMappingRendererComponent::StaticClass(), RendererCustomizationInstance);
-
 	FOnGetDetailCustomizationInstance MatrixCustomizationInstance = FOnGetDetailCustomizationInstance::CreateStatic(&FDMXPixelMappingDetailCustomization_Matrix::MakeInstance, ToolkitWeakPtr);
 	PropertyView->RegisterInstancedCustomPropertyLayout(UDMXPixelMappingMatrixComponent::StaticClass(), MatrixCustomizationInstance);
 
diff --git a/Engine/Plugins/VirtualProduction/DMX/DMXPixelMapping/Source/DMXPixelMappingEditor/Private/Views/SDMXPixelMappingInputSourceView.cpp b/Engine/Plugins/VirtualProduction/DMX/DMXPixelMapping/Source/DMXPixelMappingEditor/Private/Views/SDMXPixelMappingInputSourceView.cpp
new file mode 100644
--- /dev/null
+++ b/Engine/Plugins/VirtualProduction/DMX/DMXPixelMapping/Source/DMXPixelMappingEditor/Private/Views/SDMXPixelMappingInputSourceView.cpp
@@ -0,0 +1,292 @@
+// Copyright Epic Games, Inc. All Rights Reserved.
+
+#include "Views/SDMXPixelMappingInputSourceView.h"
+
+#include "DMXPixelMapping.h"
+#include "DMXPixelMappingComponentReference.h"
+#include "Components/DMXPixelMappingRendererComponent.h"
+#include "Components/DMXPixelMappingRootComponent.h"
+#include "Customizations/DMXPixelMappingDetailCustomization_Renderer.h"
+#include "Toolkits/DMXPixelMappingToolkit.h"
+
+#include "Editor.h"
+#include "IDetailsView.h"
+#include "PropertyEditorDelegates.h"
+#include "PropertyEditorModule.h"
+#include "Widgets/SBoxPanel.h"
+#include "Widgets/Images/SImage.h"
+#include "Widgets/Input/SButton.h"
+#include "Widgets/Layout/SExpandableArea.h"
+#include "Widgets/Layout/SScrollBox.h"
+#include "Widgets/Layout/SSeparator.h"
+#include "Widgets/Text/STextBlock.h"
+
+
+#define LOCTEXT_NAMESPACE "SDMXPixelMappingInputSourceView"
+
+void SDMXPixelMappingInputSourceView::Construct(const FArguments& InArgs, const TSharedPtr<FDMXPixelMappingToolkit>& InToolkit)
+{
+	WeakToolkit = InToolkit;
+	if (!WeakToolkit.IsValid())
+	{
+		return;
+	}
+	const TSharedRef<FDMXPixelMappingToolkit> Toolkit = WeakToolkit.Pin().ToSharedRef();
+
+	ChildSlot
+	[
+		SNew(SBorder)
+		.BorderImage(FAppStyle::GetBrush("ToolPanel.GroupBorder"))
+		[
+			SNew(SVerticalBox)
+
+			// Add button
+			+ SVerticalBox::Slot()
+			.HAlign(HAlign_Left)
+			.VAlign(VAlign_Fill)
+			.Padding(4.f)
+			.AutoHeight()
+			[
+				SNew(SButton)
+				.ButtonStyle(FAppStyle::Get(), "FlatButton.Success")
+				.ForegroundColor(FLinearColor::White)
+				.ToolTipText(LOCTEXT("AddInputSourceTooltip", "Adds a new input to the pixelmapping asset"))
+				.ContentPadding(FMargin(5.0f, 1.0f))
+				.OnClicked(this, &SDMXPixelMappingInputSourceView::OnAddButtonClicked)
+				.Content()
+				[
+					SNew(SHorizontalBox)
+
+					+ SHorizontalBox::Slot()
+					.AutoWidth()
+					.Padding(FMargin(0.f, 1.f))
+					[
+						SNew(SImage)
+						.Image(FAppStyle::GetBrush("Plus"))
+					]
+
+					+ SHorizontalBox::Slot()
+					.VAlign(VAlign_Center)
+					.AutoWidth()
+					.Padding(FMargin(2.f, 0.f, 2.f, 0.f))
+					[
+						SNew(STextBlock)
+						.Text(LOCTEXT("AddInputSourceLabel", "Add"))
+					]
+				]
+			]
+
+			// Details views 
+			+ SVerticalBox::Slot()
+			.HAlign(HAlign_Fill)
+			.VAlign(VAlign_Fill)
+			.FillHeight(1.f)
+			[
+				SAssignNew(DetailsViewsScrollBox, SScrollBox)
+				.Orientation(EOrientation::Orient_Vertical)
+			]
+		]
+	];
+
+	ForceRefresh();
+
+	UDMXPixelMappingBaseComponent::GetOnComponentAdded().AddSP(this, &SDMXPixelMappingInputSourceView::OnComponentAddedOrRemoved);
+	UDMXPixelMappingBaseComponent::GetOnComponentRemoved().AddSP(this, &SDMXPixelMappingInputSourceView::OnComponentAddedOrRemoved);
+	Toolkit->GetOnSelectedComponentsChangedDelegate().AddSP(this, &SDMXPixelMappingInputSourceView::OnComponentSelected);
+}
+
+void SDMXPixelMappingInputSourceView::RequestRefresh()
+{
+	if (!RefreshTimerHandle.IsValid())
+	{
+		RefreshTimerHandle = GEditor->GetTimerManager()->SetTimerForNextTick(FTimerDelegate::CreateSP(this, &SDMXPixelMappingInputSourceView::ForceRefresh));
+	}
+}
+
+void SDMXPixelMappingInputSourceView::ForceRefresh()
+{
+	RefreshTimerHandle.Invalidate();
+
+	if (!WeakToolkit.IsValid())
+	{
+		return;
+	}
+	const TSharedRef<FDMXPixelMappingToolkit> Toolkit = WeakToolkit.Pin().ToSharedRef();
+
+	DetailsViewsScrollBox->ClearChildren();
+	const TArray<UDMXPixelMappingRendererComponent*> RendererComponents = GetRendererComponents();
+	NumRendererComponents = RendererComponents.Num();
+	for (UDMXPixelMappingRendererComponent* RendererComponent : RendererComponents)
+	{
+		const TSharedRef<SVerticalBox> BodyContent = SNew(SVerticalBox);
+		
+		// Create a title view for each renderer component
+		const TSharedRef<SHorizontalBox> TitleContentBox = SNew(SHorizontalBox);
+			
+		TitleContentBox->AddSlot()
+			.HAlign(HAlign_Left)
+			.VAlign(VAlign_Center)
+			.FillWidth(1.f)
+			[
+				SNew(STextBlock)
+				.TextStyle(FAppStyle::Get(), "DetailsView.CategoryTextStyle")
+				.Text_Lambda([RendererComponent]()
+					{
+						return RendererComponent ? FText::FromString(RendererComponent->GetUserFriendlyName()) : FText::GetEmpty();
+					})
+			];
+
+		TitleContentBox->AddSlot()
+			.Padding(2.f, 0.f)
+			[
+				SNew(SSeparator)
+				.Orientation(EOrientation::Orient_Horizontal)
+			];
+
+		if (RendererComponents.Num() > 1)
+		{
+			TitleContentBox->AddSlot()
+				.HAlign(HAlign_Left)
+				.VAlign(VAlign_Center)
+				.Padding(4.f)
+				.AutoWidth()
+				[
+					SNew(SButton)
+					.ButtonStyle(FAppStyle::Get(), "FlatButton.Success")
+					.ForegroundColor(FLinearColor::White)
+					.ToolTipText(LOCTEXT("SelectRendererComponentTooltip", "Selects the input source"))
+					.ContentPadding(FMargin(5.0f, 1.0f))
+					.OnClicked(this, &SDMXPixelMappingInputSourceView::OnSelectRendererComponentButtonClicked, RendererComponent)
+					.Content()
+					[
+						SNew(STextBlock)
+						.Text(LOCTEXT("SelectRendererComponentLabel", "Select"))
+					]
+				];
+		}
+
+		// Create a details view for the active renderer component
+		if (Toolkit->GetActiveRendererComponent() == RendererComponent)
+		{
+			FDetailsViewArgs DetailsViewArgs;
+			DetailsViewArgs.NameAreaSettings = FDetailsViewArgs::HideNameArea;
+			DetailsViewArgs.bHideSelectionTip = true;
+			DetailsViewArgs.DefaultsOnlyVisibility = EEditDefaultsOnlyNodeVisibility::Automatic;
+
+			FPropertyEditorModule& PropertyEditorModule = FModuleManager::Get().GetModuleChecked<FPropertyEditorModule>("PropertyEditor");
+			const TSharedRef<IDetailsView> DetailsView = PropertyEditorModule.CreateDetailView(DetailsViewArgs);
+
+			FOnGetDetailCustomizationInstance RendererCustomizationInstance = FOnGetDetailCustomizationInstance::CreateStatic(&FDMXPixelMappingDetailCustomization_Renderer::MakeInstance, WeakToolkit);
+			DetailsView->RegisterInstancedCustomPropertyLayout(UDMXPixelMappingRendererComponent::StaticClass(), RendererCustomizationInstance);
+			DetailsView->SetObject(RendererComponent);
+
+			BodyContent->AddSlot()
+				.AutoHeight()
+				[
+					DetailsView
+				];
+		}
+
+		// Add to scrollbox
+		DetailsViewsScrollBox->AddSlot()
+			[
+				SNew(SExpandableArea)
+				.BorderImage(FAppStyle::Get().GetBrush("Brushes.Header"))
+				.BodyBorderImage(FAppStyle::Get().GetBrush("Brushes.Recessed"))				
+				.HeaderPadding(4.f)
+				.Padding(FMargin(8.f, 2.f, 2.f, 0.f))
+				.HeaderContent()
+				[
+					TitleContentBox
+				]
+				.BodyContent()
+				[
+					BodyContent
+				]
+			];
+	}
+}
+
+void SDMXPixelMappingInputSourceView::OnComponentSelected()
+{
+	if (const TSharedPtr<FDMXPixelMappingToolkit> Toolkit = WeakToolkit.Pin())
+	{
+		// Refresh the view, only if a renderer component is selected
+		const FDMXPixelMappingComponentReference* const SelectedRendererComponentRefPtr = Algo::FindByPredicate(Toolkit->GetSelectedComponents(), [](const FDMXPixelMappingComponentReference& ComponentRef)
+			{
+				const UDMXPixelMappingBaseComponent* Component = ComponentRef.GetComponent();
+				return Component->GetClass() == UDMXPixelMappingRendererComponent::StaticClass();
+			});
+		if (SelectedRendererComponentRefPtr)
+		{
+			RequestRefresh();
+		}
+	}
+}
+
+void SDMXPixelMappingInputSourceView::OnComponentAddedOrRemoved(UDMXPixelMapping* PixelMapping, UDMXPixelMappingBaseComponent* Component)
+{
+	RequestRefresh();
+}
+
+FReply SDMXPixelMappingInputSourceView::OnAddButtonClicked()
+{
+	if (const TSharedPtr<FDMXPixelMappingToolkit> Toolkit = WeakToolkit.Pin())
+	{
+		Toolkit->AddRenderer();
+		
+		if (UDMXPixelMappingRendererComponent* RendererComponentToSelect = Toolkit->GetActiveRendererComponent())
+		{
+			const FDMXPixelMappingComponentReference ComponentRefToSelect(Toolkit, RendererComponentToSelect);
+			const TSet<FDMXPixelMappingComponentReference> NewSelection{ ComponentRefToSelect };
+			Toolkit->SelectComponents(NewSelection);
+			Toolkit->ZoomToFit();
+		}
+
+		RequestRefresh();
+		DetailsViewsScrollBox->ScrollToEnd();
+	}
+
+	return FReply::Handled();
+}
+
+FReply SDMXPixelMappingInputSourceView::OnSelectRendererComponentButtonClicked(UDMXPixelMappingRendererComponent* RendererComponent)
+{
+	const TSharedPtr<FDMXPixelMappingToolkit> Toolkit = WeakToolkit.Pin();
+	if (Toolkit.IsValid() && RendererComponent)
+	{
+		const FDMXPixelMappingComponentReference ComponentRefToSelect(Toolkit, RendererComponent);
+		const TSet<FDMXPixelMappingComponentReference> NewSelection{ ComponentRefToSelect };
+		Toolkit->SelectComponents(NewSelection);
+	}
+
+	return FReply::Handled();
+}
+
+TArray<UDMXPixelMappingRendererComponent*> SDMXPixelMappingInputSourceView::GetRendererComponents() const
+{
+	TArray<UDMXPixelMappingRendererComponent*> Result;
+	if (!WeakToolkit.IsValid())
+	{
+		return Result;
+	}
+	const TSharedRef<FDMXPixelMappingToolkit> Toolkit = WeakToolkit.Pin().ToSharedRef();
+
+	UDMXPixelMappingRootComponent* RootComponent = Toolkit->GetDMXPixelMapping() ? Toolkit->GetDMXPixelMapping()->GetRootComponent() : nullptr;
+	if (!RootComponent)
+	{
+		return Result;
+	}
+
+	for (UDMXPixelMappingBaseComponent* Child : RootComponent->GetChildren())
+	{
+		if (UDMXPixelMappingRendererComponent* RendererComponent = Cast<UDMXPixelMappingRendererComponent>(Child))
+		{
+			Result.Add(RendererComponent);
+		}
+	}
+
+	return Result;
+}
+
+#undef LOCTEXT_NAMESPACE
diff --git a/Engine/Plugins/VirtualProduction/DMX/DMXPixelMapping/Source/DMXPixelMappingEditor/Private/Views/SDMXPixelMappingInputSourceView.h b/Engine/Plugins/VirtualProduction/DMX/DMXPixelMapping/Source/DMXPixelMappingEditor/Private/Views/SDMXPixelMappingInputSourceView.h
new file mode 100644
--- /dev/null
+++ b/Engine/Plugins/VirtualProduction/DMX/DMXPixelMapping/Source/DMXPixelMappingEditor/Private/Views/SDMXPixelMappingInputSourceView.h
@@ -0,0 +1,63 @@
+// Copyright Epic Games, Inc. All Rights Reserved.
+
+#pragma once
+
+#include "Engine/EngineTypes.h"
+#include "Widgets/SCompoundWidget.h"
+
+class UDMXPixelMapping;
+class UDMXPixelMappingBaseComponent;
+class FDMXPixelMappingToolkit;
+class UDMXPixelMappingRendererComponent;
+
+class IDetailsView;
+class FReply;
+class SScrollBox;
+class SVerticalBox;
+
+
+/** Displays the input sources of the pixel mapping */
+class SDMXPixelMappingInputSourceView
+	: public SCompoundWidget
+{
+public:
+	SLATE_BEGIN_ARGS(SDMXPixelMappingInputSourceView) { }
+	SLATE_END_ARGS()
+
+	/** Constructs this widget */
+	void Construct(const FArguments& InArgs, const TSharedPtr<FDMXPixelMappingToolkit>& InToolkit);
+
+	/** Refreshes the widget on the next tick */
+	void RequestRefresh();
+
+private:
+	/** Refreshes the widget direct */
+	void ForceRefresh();
+
+	/** Called when a component was selected */
+	void OnComponentSelected();
+
+	/** Called when a component was added to or removed from the pixel mapping */
+	void OnComponentAddedOrRemoved(UDMXPixelMapping* PixelMapping, UDMXPixelMappingBaseComponent* Component);
+
+	/** Called when the add button was clicked */
+	FReply OnAddButtonClicked();
+
+	/** Called when the select renederer component button was clicked*/
+	FReply OnSelectRendererComponentButtonClicked(UDMXPixelMappingRendererComponent* RendererComponent);
+
+	/** Number of renderer components in the pixel mapping */
+	int32 NumRendererComponents = 0;
+
+	/** Returns the renderer components in this pixelmapping */
+	TArray<UDMXPixelMappingRendererComponent*> GetRendererComponents() const;
+
+	/** Scrollbox that holds the details views */
+	TSharedPtr<SScrollBox> DetailsViewsScrollBox;
+
+	/** Timer handle for the Request Refresh method */
+	FTimerHandle RefreshTimerHandle;
+
+	/** The toolkit of this editor */
+	TWeakPtr<FDMXPixelMappingToolkit> WeakToolkit;
+};
diff --git a/Engine/Plugins/VirtualProduction/DMX/DMXPixelMapping/Source/DMXPixelMappingEditor/Private/Views/SDMXPixelMappingPaletteView.cpp b/Engine/Plugins/VirtualProduction/DMX/DMXPixelMapping/Source/DMXPixelMappingEditor/Private/Views/SDMXPixelMappingPaletteView.cpp
deleted file mode 100644
diff --git a/Engine/Plugins/VirtualProduction/DMX/DMXPixelMapping/Source/DMXPixelMappingEditor/Private/Views/SDMXPixelMappingPaletteView.h b/Engine/Plugins/VirtualProduction/DMX/DMXPixelMapping/Source/DMXPixelMappingEditor/Private/Views/SDMXPixelMappingPaletteView.h
deleted file mode 100644
diff --git a/Engine/Plugins/VirtualProduction/DMX/DMXPixelMapping/Source/DMXPixelMappingEditor/Private/Widgets/SDMXPixelMappingFixturePatchDetailRow.cpp b/Engine/Plugins/VirtualProduction/DMX/DMXPixelMapping/Source/DMXPixelMappingEditor/Private/Widgets/SDMXPixelMappingFixturePatchDetailRow.cpp
deleted file mode 100644
diff --git a/Engine/Plugins/VirtualProduction/DMX/DMXPixelMapping/Source/DMXPixelMappingEditor/Private/Widgets/SDMXPixelMappingFixturePatchDetailRow.h b/Engine/Plugins/VirtualProduction/DMX/DMXPixelMapping/Source/DMXPixelMappingEditor/Private/Widgets/SDMXPixelMappingFixturePatchDetailRow.h
deleted file mode 100644
diff --git a/Engine/Plugins/VirtualProduction/DMX/DMXPixelMapping/Source/DMXPixelMappingEditor/Private/Widgets/SDMXPixelMappingFixturePatchList.cpp b/Engine/Plugins/VirtualProduction/DMX/DMXPixelMapping/Source/DMXPixelMappingEditor/Private/Widgets/SDMXPixelMappingFixturePatchList.cpp
new file mode 100644
--- /dev/null
+++ b/Engine/Plugins/VirtualProduction/DMX/DMXPixelMapping/Source/DMXPixelMappingEditor/Private/Widgets/SDMXPixelMappingFixturePatchList.cpp
@@ -0,0 +1,394 @@
+// Copyright Epic Games, Inc. All Rights Reserved.
+
+#include "SDMXPixelMappingFixturePatchList.h"
+
+#include "DMXPixelMapping.h"
+#include "Components/DMXPixelMappingFixtureGroupComponent.h"
+#include "Components/DMXPixelMappingFixtureGroupItemComponent.h"
+#include "Components/DMXPixelMappingMatrixComponent.h"
+#include "DragDrop/DMXPixelMappingDragDropOp.h"
+#include "Library/DMXEntityFixturePatch.h"
+#include "Library/DMXEntityFixtureType.h"
+#include "Library/DMXLibrary.h"
+#include "Templates/DMXPixelMappingComponentTemplate.h"
+#include "Toolkits/DMXPixelMappingToolkit.h"
+#include "ViewModels/DMXPixelMappingDMXLibraryViewModel.h"
+#include "Widgets/SDMXReadOnlyFixturePatchList.h"
+
+#include "ScopedTransaction.h"
+#include "Algo/Find.h"
+#include "Widgets/Input/SButton.h"
+
+
+#define LOCTEXT_NAMESPACE "SDMXPixelMappingFixturePatchList"
+
+SDMXPixelMappingFixturePatchList::~SDMXPixelMappingFixturePatchList()
+{
+	if (WeakViewModel.IsValid() && FixturePatchList.IsValid())
+	{
+		WeakViewModel->SaveFixturePatchListDescriptor(FixturePatchList->GetListDescriptor());
+	}
+}
+
+void SDMXPixelMappingFixturePatchList::Construct(const FArguments& InArgs, const TSharedPtr<FDMXPixelMappingToolkit>& InToolkit, UDMXPixelMappingDMXLibraryViewModel* InViewModel)
+{
+	if (!InToolkit.IsValid() || !InViewModel || !InViewModel->GetDMXLibrary())
+	{
+		return;
+	}
+	WeakToolkit = InToolkit;
+	WeakViewModel = InViewModel;
+
+	ChildSlot
+		[
+			SNew(SVerticalBox)
+
+			+ SVerticalBox::Slot()
+			.AutoHeight()
+			[
+				SNew(SHorizontalBox)
+				
+				+ SHorizontalBox::Slot()
+				.Padding(4.f)
+				.AutoWidth()
+				[
+					SNew(SButton)
+					.ButtonStyle(FAppStyle::Get(), "FlatButton.Success")
+					.ForegroundColor(FLinearColor::White)
+					.ToolTipText(LOCTEXT("AddSelectedPatchesTooltip", "Adds the selected patches to the Pixel Mapping"))
+					.ContentPadding(FMargin(5.0f, 1.0f))
+					.OnClicked(this, &SDMXPixelMappingFixturePatchList::OnAddSelectedPatchesClicked)
+					.Content()
+					[
+						SNew(STextBlock)
+						.Text(LOCTEXT("AddSelectedPatchesLabel", "Add Selected Patches"))
+					]
+				]
+
+				+ SHorizontalBox::Slot()
+				.Padding(8.f, 4.f)
+				.AutoWidth()
+				[
+					SNew(SButton)
+					.ButtonStyle(FAppStyle::Get(), "FlatButton.Success")
+					.ForegroundColor(FLinearColor::White)
+					.ToolTipText(LOCTEXT("AddAllPatchesTooltip", "Adds all patches to the Pixel Mapping"))
+					.ContentPadding(FMargin(5.0f, 1.0f))
+					.OnClicked(this, &SDMXPixelMappingFixturePatchList::OnAddAllPatchesClicked)
+					.Content()
+					[
+						SNew(STextBlock)
+						.Text(LOCTEXT("AddAllPatchesLabel", "Add all Patches"))
+					]
+				]
+			]
+
+			+ SVerticalBox::Slot()
+			.AutoHeight()
+			[
+				SAssignNew(FixturePatchList, SDMXReadOnlyFixturePatchList)
+				.DMXLibrary(InViewModel->GetDMXLibrary())
+				.ExcludedFixturePatches(InViewModel->GetFixturePatchesInUse())
+				.ListDescriptor(InViewModel->GetFixturePatchListDescriptor())
+				.OnRowDragged(this, &SDMXPixelMappingFixturePatchList::OnRowDragged)
+			]
+		];
+
+	
+	// Make an initial selection
+	if (FixturePatchList.IsValid() && !FixturePatchList->GetListItems().IsEmpty())
+	{
+		const TArray<TSharedPtr<FDMXEntityFixturePatchRef>> NewSelection{ FixturePatchList->GetListItems()[0] };
+		FixturePatchList->SelectItems(NewSelection);
+	}
+
+	UDMXPixelMappingBaseComponent::GetOnComponentAdded().AddSP(this, &SDMXPixelMappingFixturePatchList::OnComponentAddedOrRemoved);
+	UDMXPixelMappingBaseComponent::GetOnComponentRemoved().AddSP(this, &SDMXPixelMappingFixturePatchList::OnComponentAddedOrRemoved);
+}
+
+void SDMXPixelMappingFixturePatchList::PostUndo(bool bSuccess)
+{
+	if (WeakViewModel.IsValid())
+	{
+		FixturePatchList->SetExcludedFixturePatches(WeakViewModel->GetFixturePatchesInUse());
+	}
+}
+
+void SDMXPixelMappingFixturePatchList::PostRedo(bool bSuccess)
+{
+	if (WeakViewModel.IsValid())
+	{
+		FixturePatchList->SetExcludedFixturePatches(WeakViewModel->GetFixturePatchesInUse());
+	}
+}
+
+FReply SDMXPixelMappingFixturePatchList::OnRowDragged(const FGeometry& MyGeometry, const FPointerEvent& MouseEvent)
+{
+	UDMXPixelMappingFixtureGroupComponent* FixtureGroupComponent = WeakViewModel.IsValid() ? WeakViewModel->GetFixtureGroupComponent() : nullptr;
+	if (!FixtureGroupComponent)
+	{
+		return FReply::Unhandled();
+	}
+
+	const TArray<TSharedPtr<FDMXEntityFixturePatchRef>> SelectedPatches = FixturePatchList->GetSelectedFixturePatchRefs();
+	if (!SelectedPatches.IsEmpty())
+	{
+		TArray<TSharedPtr<FDMXPixelMappingComponentTemplate>> Templates;
+		for (const TSharedPtr<FDMXEntityFixturePatchRef>& FixturePatchRef : SelectedPatches)
+		{
+			if (!FixturePatchRef.IsValid())
+			{
+				continue;
+			}
+
+			const UDMXEntityFixturePatch* FixturePatch = FixturePatchRef->GetFixturePatch();
+			const UDMXEntityFixtureType* FixtureType = FixturePatch ? FixturePatch->GetFixtureType() : nullptr;
+			const FDMXFixtureMode* ActiveModePtr = FixturePatch ? FixturePatch->GetActiveMode() : nullptr;
+			if (FixturePatch && FixtureType && ActiveModePtr)
+			{
+				if (ActiveModePtr->bFixtureMatrixEnabled)
+				{
+					TSharedRef<FDMXPixelMappingComponentTemplate> FixturePatchMatrixTemplate = MakeShared<FDMXPixelMappingComponentTemplate>(UDMXPixelMappingMatrixComponent::StaticClass(), *FixturePatchRef);
+					Templates.Add(FixturePatchMatrixTemplate);
+				}
+				else
+				{
+					TSharedRef<FDMXPixelMappingComponentTemplate> FixturePatchItemTemplate = MakeShared<FDMXPixelMappingComponentTemplate>(UDMXPixelMappingFixtureGroupItemComponent::StaticClass(), *FixturePatchRef);
+					Templates.Add(FixturePatchItemTemplate);
+				}
+			}
+		}
+		return FReply::Handled().BeginDragDrop(FDMXPixelMappingDragDropOp::New(FVector2D::ZeroVector, Templates, FixtureGroupComponent));
+	}
+
+	return FReply::Unhandled();
+}
+
+void SDMXPixelMappingFixturePatchList::OnComponentAddedOrRemoved(UDMXPixelMapping* PixelMapping, UDMXPixelMappingBaseComponent* Component)
+{
+	FixturePatchList->RequestListRefresh();
+}
+
+FReply SDMXPixelMappingFixturePatchList::OnAddSelectedPatchesClicked()
+{
+	if (WeakViewModel.IsValid())
+	{
+		const FScopedTransaction AddSelectedFixturePatchesTransaction(LOCTEXT("AddSelectedFixturePatchesTransaction", "Add selected Fixture Patches"));
+
+		const TArray<TSharedPtr<FDMXEntityFixturePatchRef>> SelectedPatches = FixturePatchList->GetSelectedFixturePatchRefs();
+		
+		if (!SelectedPatches.IsEmpty())
+		{
+			AddFixturePatches(SelectedPatches);
+
+			SelectAfter(SelectedPatches.Last()->GetFixturePatch());
+			FixturePatchList->SetExcludedFixturePatches(WeakViewModel->GetFixturePatchesInUse());
+		}
+	}
+
+	return FReply::Handled();
+}
+
+FReply SDMXPixelMappingFixturePatchList::OnAddAllPatchesClicked()
+{
+	if (WeakViewModel.IsValid())
+	{
+		const FScopedTransaction AddAllFixturePatchesTransaction(LOCTEXT("AddAllFixturePatchesTransaction", "Add all Fixture Patches"));
+
+		const TArray<TSharedPtr<FDMXEntityFixturePatchRef>> AllVisiblePatches = FixturePatchList->GetVisibleFixturePatchRefs();
+
+		if (!AllVisiblePatches.IsEmpty())
+		{
+			AddFixturePatches(AllVisiblePatches);
+
+			FixturePatchList->SetExcludedFixturePatches(WeakViewModel->GetFixturePatchesInUse());
+		}
+	}
+
+	return FReply::Handled();
+}
+
+void SDMXPixelMappingFixturePatchList::SelectAfter(const UDMXEntityFixturePatch* FixturePatch)
+{
+	if (!FixturePatch)
+	{
+		return;
+	}
+
+	const TArray<TSharedPtr<FDMXEntityFixturePatchRef>> Items = FixturePatchList->GetListItems();
+	int32 IndexOfPatch = Items.IndexOfByPredicate([FixturePatch](const TSharedPtr<FDMXEntityFixturePatchRef>& Item)
+		{
+			return Item->GetFixturePatch() == FixturePatch;
+		});
+	if (Items.IsValidIndex(IndexOfPatch + 1))
+	{
+		const TArray<TSharedPtr<FDMXEntityFixturePatchRef>> NewSelection{ Items[IndexOfPatch + 1]};
+		FixturePatchList->SelectItems(NewSelection);
+	}
+	else if (Items.IsValidIndex(IndexOfPatch - 1))
+	{
+		const TArray<TSharedPtr<FDMXEntityFixturePatchRef>> NewSelection{ Items[IndexOfPatch - 1] };
+		FixturePatchList->SelectItems(NewSelection);
+	}
+}
+
+void SDMXPixelMappingFixturePatchList::AddFixturePatches(const TArray<TSharedPtr<FDMXEntityFixturePatchRef>>& FixturePatches)
+{
+	const TSharedPtr<FDMXPixelMappingToolkit> Toolkit = WeakToolkit.Pin();
+	UDMXPixelMappingDMXLibraryViewModel* ViewModel = WeakViewModel.Get();
+	UDMXPixelMappingFixtureGroupComponent* FixtureGroupComponent = ViewModel ? ViewModel->GetFixtureGroupComponent() : nullptr;
+	UDMXPixelMappingRootComponent* RootComponent = ViewModel ? ViewModel->GetPixelMappingRootComponent() : nullptr;
+	if (!Toolkit.IsValid() || !ViewModel || !FixtureGroupComponent || !RootComponent)
+	{
+		return;
+	}
+
+	TArray<TSharedPtr<FDMXPixelMappingComponentTemplate>> Templates;
+	for (const TSharedPtr<FDMXEntityFixturePatchRef>& FixturePatchRef : FixturePatches)
+	{
+		if (!FixturePatchRef.IsValid())
+		{
+			continue;
+		}
+
+		UDMXEntityFixturePatch* FixturePatch = FixturePatchRef->GetFixturePatch();
+		UDMXEntityFixtureType* FixtureType = FixturePatch ? FixturePatch->GetFixtureType() : nullptr;
+		const FDMXFixtureMode* ActiveModePtr = FixturePatch ? FixturePatch->GetActiveMode() : nullptr;
+		if (FixturePatch && FixtureType && ActiveModePtr)
+		{
+			if (ActiveModePtr->bFixtureMatrixEnabled)
+			{
+				const TSharedRef<FDMXPixelMappingComponentTemplate> FixturePatchMatrixTemplate = MakeShared<FDMXPixelMappingComponentTemplate>(UDMXPixelMappingMatrixComponent::StaticClass(), *FixturePatchRef);
+				Templates.Add(FixturePatchMatrixTemplate);
+			}
+			else
+			{
+				const TSharedRef<FDMXPixelMappingComponentTemplate> FixturePatchItemTemplate = MakeShared<FDMXPixelMappingComponentTemplate>(UDMXPixelMappingFixtureGroupItemComponent::StaticClass(), *FixturePatchRef);
+				Templates.Add(FixturePatchItemTemplate);
+			}
+		}
+	}
+
+	const TArray<UDMXPixelMappingBaseComponent*> NewComponents = Toolkit->CreateComponentsFromTemplates(RootComponent, FixtureGroupComponent, Templates);
+
+	// Layout
+	const TArray<UDMXEntityFixturePatch*> FixturePatchesInLibrary = ViewModel->GetDMXLibrary()->GetEntitiesTypeCast<UDMXEntityFixturePatch>();
+	bool bLayoutEvenOverParent = FixturePatchesInLibrary.Num() == FixturePatches.Num();
+	if (bLayoutEvenOverParent)
+	{
+		LayoutEvenOverParent(NewComponents);
+	}
+	else
+	{
+		LayoutAfterLastPatch(NewComponents);
+	}
+}
+
+void SDMXPixelMappingFixturePatchList::LayoutEvenOverParent(const TArray<UDMXPixelMappingBaseComponent*> Components)
+{
+	UDMXPixelMappingDMXLibraryViewModel* ViewModel = WeakViewModel.Get();
+	UDMXPixelMappingFixtureGroupComponent* FixtureGroupComponent = ViewModel ? ViewModel->GetFixtureGroupComponent() : nullptr;
+	if (!ViewModel || !FixtureGroupComponent)
+	{
+		return;
+	}
+
+	const int32 Columns = FMath::RoundFromZero(FMath::Sqrt((float)Components.Num()));
+	const int32 Rows = FMath::RoundFromZero((float)Components.Num() / Columns);
+	const FVector2D Size = FVector2D(FixtureGroupComponent->GetSize().X / Columns, FixtureGroupComponent->GetSize().Y / Rows);
+
+	const FVector2D ParentPosition = FixtureGroupComponent->GetPosition();
+	int32 Column = -1;
+	int32 Row = 0;
+	for (UDMXPixelMappingBaseComponent* Component : Components)
+	{
+		if (UDMXPixelMappingOutputComponent* OutputComponent = Cast<UDMXPixelMappingOutputComponent>(Component))
+		{
+			Column++;
+			if (Column >= Columns)
+			{
+				Column = 0;
+				Row++;
+			}
+			const FVector2D Position = ParentPosition + FVector2D(Column * Size.X, Row * Size.Y);
+			OutputComponent->SetPosition(Position);
+			OutputComponent->SetSize(Size);
+		}
+	}
+}
+
+void SDMXPixelMappingFixturePatchList::LayoutAfterLastPatch(const TArray<UDMXPixelMappingBaseComponent*> Components)
+{
+	UDMXPixelMappingDMXLibraryViewModel* ViewModel = WeakViewModel.Get();
+	UDMXPixelMappingFixtureGroupComponent* FixtureGroupComponent = ViewModel ? ViewModel->GetFixtureGroupComponent() : nullptr;
+	UDMXPixelMappingOutputComponent* FirstComponentToLayout = Components.IsEmpty() ? nullptr : Cast<UDMXPixelMappingOutputComponent>(Components[0]);
+	if (!ViewModel || !FixtureGroupComponent || !FirstComponentToLayout)
+	{
+		return;
+	}
+
+	// Find other components
+	TArray<UDMXPixelMappingBaseComponent*> OtherComponentsInGroup;
+	for (UDMXPixelMappingBaseComponent* Child : FixtureGroupComponent->GetChildren())
+	{
+		if (!Components.Contains(Child))
+		{
+			OtherComponentsInGroup.Add(Child);
+		}
+	}
+
+	// Find a starting position
+	FVector2D NextPosition = FixtureGroupComponent->GetPosition();
+	for (UDMXPixelMappingBaseComponent* OtherComponent : OtherComponentsInGroup)
+	{
+		if (UDMXPixelMappingOutputComponent* OtherOutputComponent = Cast<UDMXPixelMappingOutputComponent>(OtherComponent))
+		{
+			NextPosition.X = FMath::Max(NextPosition.X, OtherOutputComponent->GetPosition().X + OtherOutputComponent->GetSize().X + 1);
+			if (NextPosition.X + FirstComponentToLayout->GetSize().X > FixtureGroupComponent->GetPosition().X + FixtureGroupComponent->GetSize().X - 1)
+			{
+				NextPosition.X = FixtureGroupComponent->GetPosition().X;
+				NextPosition.Y += FirstComponentToLayout->GetSize().Y;
+			}
+		}
+	}
+
+	// Layout
+	float RowHeight = 0.f;
+	for (UDMXPixelMappingBaseComponent* Component : Components)
+	{
+		if (UDMXPixelMappingOutputComponent* OutputComponent = Cast<UDMXPixelMappingOutputComponent>(Component))
+		{
+			if (FixtureGroupComponent->IsOverPosition(NextPosition) &&
+				FixtureGroupComponent->IsOverPosition(NextPosition + OutputComponent->GetSize()))
+			{
+				OutputComponent->SetPosition(NextPosition);
+
+				RowHeight = FMath::Max(OutputComponent->GetSize().Y, RowHeight);
+				NextPosition = FVector2D(NextPosition.X + OutputComponent->GetSize().X, NextPosition.Y);
+			}
+			else
+			{
+				// Try on a new row
+				FVector2D NewRowPosition = FVector2D(FixtureGroupComponent->GetPosition().X, NextPosition.Y + RowHeight);
+
+				const FVector2D NextPositionOnNewRow = FVector2D(NewRowPosition.X + OutputComponent->GetSize().X, NewRowPosition.Y);
+				if (FixtureGroupComponent->IsOverPosition(NextPositionOnNewRow) &&
+					FixtureGroupComponent->IsOverPosition(NextPositionOnNewRow + OutputComponent->GetSize()))
+				{
+					OutputComponent->SetPosition(NewRowPosition);
+
+					NextPosition = FVector2D(NewRowPosition.X + OutputComponent->GetSize().X, NewRowPosition.Y);
+					RowHeight = OutputComponent->GetSize().Y;
+				}
+				else
+				{
+					OutputComponent->SetPosition(NextPosition);
+
+					NextPosition = FVector2D(NextPosition.X + OutputComponent->GetSize().X, NextPosition.Y);
+				}
+			}
+		}
+	}
+}
+
+#undef LOCTEXT_NAMESPACE
diff --git a/Engine/Plugins/VirtualProduction/DMX/DMXPixelMapping/Source/DMXPixelMappingEditor/Private/Widgets/SDMXPixelMappingFixturePatchList.h b/Engine/Plugins/VirtualProduction/DMX/DMXPixelMapping/Source/DMXPixelMappingEditor/Private/Widgets/SDMXPixelMappingFixturePatchList.h
new file mode 100644
--- /dev/null
+++ b/Engine/Plugins/VirtualProduction/DMX/DMXPixelMapping/Source/DMXPixelMappingEditor/Private/Widgets/SDMXPixelMappingFixturePatchList.h
@@ -0,0 +1,77 @@
+// Copyright Epic Games, Inc. All Rights Reserved.
+
+#pragma once
+
+#include "Library/DMXEntityReference.h"
+
+#include "EditorUndoClient.h"
+#include "Widgets/SCompoundWidget.h"
+
+class FDMXPixelMappingToolkit;
+class SDMXReadOnlyFixturePatchList;
+class UDMXEntityFixturePatch;
+class UDMXPixelMapping;
+class UDMXPixelMappingBaseComponent;
+class UDMXPixelMappingDMXLibraryViewModel;
+class UDMXLibrary;
+
+class FReply;
+
+
+/** Displays the DMX Library of the currently selected fixture group component */
+class SDMXPixelMappingFixturePatchList final
+	: public SCompoundWidget
+	, public FEditorUndoClient
+{
+public:
+	SLATE_BEGIN_ARGS(SDMXPixelMappingFixturePatchList) 
+	{}
+
+	SLATE_END_ARGS()
+
+	/** Destructor */
+	~SDMXPixelMappingFixturePatchList();
+
+	/** Constructs this widget */
+	void Construct(const FArguments& InArgs, const TSharedPtr<FDMXPixelMappingToolkit>& InToolkit, UDMXPixelMappingDMXLibraryViewModel* InViewModel);
+
+protected:
+	//~ Begin FEditorUndoClient Interface
+	virtual void PostUndo(bool bSuccess) override;
+	virtual void PostRedo(bool bSuccess) override;
+	// End of FEditorUndoClient
+
+private:
+	/** Called when a row was dragged */
+	FReply OnRowDragged(const FGeometry& MyGeometry, const FPointerEvent& MouseEvent);
+
+	/** Called when a component was added to or removed from the pixel mapping */
+	void OnComponentAddedOrRemoved(UDMXPixelMapping* PixelMapping, UDMXPixelMappingBaseComponent* Component);
+	
+	/** Called when the 'Add selected Patches' button was clicked */
+	FReply OnAddSelectedPatchesClicked();
+
+	/** Called when the 'Add all Patches' button was clicked */
+	FReply OnAddAllPatchesClicked();
+
+	/** Selects the first patch after the specified fixture patch */
+	void SelectAfter(const UDMXEntityFixturePatch* FixturePatch);
+
+	/** Adds specified fixture patches to the pixel mapping */
+	void AddFixturePatches(const TArray<TSharedPtr<FDMXEntityFixturePatchRef>>& FixturePatches);
+
+	/** Layouts components evenly over the parent group */
+	void LayoutEvenOverParent(const TArray<UDMXPixelMappingBaseComponent*> Components);
+
+	/** Layouts components after the last patch in the parent group */
+	void LayoutAfterLastPatch(const TArray<UDMXPixelMappingBaseComponent*> Components);
+
+	/** The fixture patch list displayed in this widget */
+	TSharedPtr<SDMXReadOnlyFixturePatchList> FixturePatchList;
+
+	/** The view model used in this widget */
+	TWeakObjectPtr<UDMXPixelMappingDMXLibraryViewModel> WeakViewModel;
+
+	/** The toolkit of the editor that displays this widget */
+	TWeakPtr<FDMXPixelMappingToolkit> WeakToolkit;
+};
diff --git a/Engine/Plugins/VirtualProduction/DMX/DMXPixelMapping/Source/DMXPixelMappingEditor/Private/Widgets/SDMXPixelMappingPaletteItem.cpp b/Engine/Plugins/VirtualProduction/DMX/DMXPixelMapping/Source/DMXPixelMappingEditor/Private/Widgets/SDMXPixelMappingPaletteItem.cpp
deleted file mode 100644
diff --git a/Engine/Plugins/VirtualProduction/DMX/DMXPixelMapping/Source/DMXPixelMappingEditor/Private/Widgets/SDMXPixelMappingPaletteItem.h b/Engine/Plugins/VirtualProduction/DMX/DMXPixelMapping/Source/DMXPixelMappingEditor/Private/Widgets/SDMXPixelMappingPaletteItem.h
deleted file mode 100644
diff --git a/Engine/Plugins/VirtualProduction/DMX/DMXPixelMapping/Source/DMXPixelMappingEditor/Private/Widgets/SDMXPixelMappingSurface.h b/Engine/Plugins/VirtualProduction/DMX/DMXPixelMapping/Source/DMXPixelMappingEditor/Private/Widgets/SDMXPixelMappingSurface.h
--- a/Engine/Plugins/VirtualProduction/DMX/DMXPixelMapping/Source/DMXPixelMappingEditor/Private/Widgets/SDMXPixelMappingSurface.h
+++ b/Engine/Plugins/VirtualProduction/DMX/DMXPixelMapping/Source/DMXPixelMappingEditor/Private/Widgets/SDMXPixelMappingSurface.h
@@ -51,6 +51,9 @@ public:
 	/** Returns the toolkit used with this surface */
 	FORCEINLINE TSharedPtr<FDMXPixelMappingToolkit> GetToolkit() { return ToolkitWeakPtr.Pin(); }
 
+	/** Zooms to fit content */
+	void ZoomToFit(bool bInstantZoom);
+
 protected:
 	virtual void OnPaintBackground(const FGeometry& AllottedGeometry, const FSlateRect& MyCullingRect, FSlateWindowElementList& OutDrawElements, int32 LayerId) const;
 
@@ -64,8 +67,6 @@ protected:
 
 	bool ZoomToLocation(const FVector2D& CurrentSizeWithoutZoom, const FVector2D& DesiredSize, bool bDoneScrolling);
 
-	void ZoomToFit(bool bInstantZoom);
-
 	FText GetZoomText() const;
 	FSlateColor GetZoomTextColorAndOpacity() const;
 
diff --git a/Engine/Plugins/VirtualProduction/DMX/DMXPixelMapping/Source/DMXPixelMappingRuntime/Private/Components/DMXPixelMappingFixtureGroupComponent.cpp b/Engine/Plugins/VirtualProduction/DMX/DMXPixelMapping/Source/DMXPixelMappingRuntime/Private/Components/DMXPixelMappingFixtureGroupComponent.cpp
--- a/Engine/Plugins/VirtualProduction/DMX/DMXPixelMapping/Source/DMXPixelMappingRuntime/Private/Components/DMXPixelMappingFixtureGroupComponent.cpp
+++ b/Engine/Plugins/VirtualProduction/DMX/DMXPixelMapping/Source/DMXPixelMappingRuntime/Private/Components/DMXPixelMappingFixtureGroupComponent.cpp
@@ -38,22 +38,27 @@ void UDMXPixelMappingFixtureGroupComponent::PostEditChangeProperty(FPropertyChan
 	// Call the parent at the first place
 	Super::PostEditChangeProperty(PropertyChangedEvent);
 
-	if (PropertyChangedEvent.GetPropertyName() == UDMXPixelMappingOutputComponent::GetPositionXPropertyName() ||
-		PropertyChangedEvent.GetPropertyName() == UDMXPixelMappingOutputComponent::GetPositionYPropertyName())
+	const FName PropertyName = PropertyChangedEvent.GetPropertyName();
+	if (PropertyName == UDMXPixelMappingOutputComponent::GetPositionXPropertyName() ||
+		PropertyName == UDMXPixelMappingOutputComponent::GetPositionYPropertyName())
 	{
 		HandlePositionChanged();
 	}
+	else if (PropertyName == GET_MEMBER_NAME_CHECKED(UDMXPixelMappingFixtureGroupComponent, DMXLibrary))
+	{
+		OnDMXLibraryChangedDelegate.Broadcast();
+	}
 
 	PRAGMA_DISABLE_DEPRECATION_WARNINGS
-	if (PropertyChangedEvent.GetPropertyName() == UDMXPixelMappingOutputComponent::GetSizeXPropertyName() ||
-		PropertyChangedEvent.GetPropertyName() == UDMXPixelMappingOutputComponent::GetSizeYPropertyName())
+	if (PropertyName == UDMXPixelMappingOutputComponent::GetSizeXPropertyName() ||
+		PropertyName == UDMXPixelMappingOutputComponent::GetSizeYPropertyName())
 	{
 		if (ComponentWidget_DEPRECATED.IsValid())
 		{
 			ComponentWidget_DEPRECATED->SetSize(GetSize());
 		}
 	}
-	else if (PropertyChangedEvent.GetPropertyName() == GET_MEMBER_NAME_CHECKED(UDMXPixelMappingFixtureGroupComponent, DMXLibrary))
+	else if (PropertyName == GET_MEMBER_NAME_CHECKED(UDMXPixelMappingFixtureGroupComponent, DMXLibrary))
 	{
 		if (ComponentWidget_DEPRECATED.IsValid())
 		{
diff --git a/Engine/Plugins/VirtualProduction/DMX/DMXPixelMapping/Source/DMXPixelMappingRuntime/Public/Components/DMXPixelMappingFixtureGroupComponent.h b/Engine/Plugins/VirtualProduction/DMX/DMXPixelMapping/Source/DMXPixelMappingRuntime/Public/Components/DMXPixelMappingFixtureGroupComponent.h
--- a/Engine/Plugins/VirtualProduction/DMX/DMXPixelMapping/Source/DMXPixelMappingRuntime/Public/Components/DMXPixelMappingFixtureGroupComponent.h
+++ b/Engine/Plugins/VirtualProduction/DMX/DMXPixelMapping/Source/DMXPixelMappingRuntime/Public/Components/DMXPixelMappingFixtureGroupComponent.h
@@ -60,6 +60,9 @@ public:
 	UPROPERTY(EditAnywhere, Instanced, Category = "Layout")
 	TObjectPtr<UDMXPixelMappingLayoutScript> LayoutScript;
 
+	/** Returns a delegate broadcasted when the DMX Library changed */
+	FSimpleMulticastDelegate& GetOnDMXLibraryChanged() { return OnDMXLibraryChangedDelegate; }
+
 private:
 	/** Handles changes in position */
 	void HandlePositionChanged();
@@ -67,5 +70,8 @@ private:
 	/** Holds the last set size */
 	FVector2D LastPosition;
 
+	/** Returns a delegate broadcasted when the DMX Library changed */
+	FSimpleMulticastDelegate OnDMXLibraryChangedDelegate;
+
 	static const FVector2D MinGroupSize;
 };
