DMX - Fix CIS Static Analysis Win64

diff --git a/Engine/Plugins/VirtualProduction/DMX/DMXPixelMapping/Source/DMXPixelMappingEditor/Private/Views/SDMXPixelMappingDMXLibraryView.cpp b/Engine/Plugins/VirtualProduction/DMX/DMXPixelMapping/Source/DMXPixelMappingEditor/Private/Views/SDMXPixelMappingDMXLibraryView.cpp
--- a/Engine/Plugins/VirtualProduction/DMX/DMXPixelMapping/Source/DMXPixelMappingEditor/Private/Views/SDMXPixelMappingDMXLibraryView.cpp
+++ b/Engine/Plugins/VirtualProduction/DMX/DMXPixelMapping/Source/DMXPixelMappingEditor/Private/Views/SDMXPixelMappingDMXLibraryView.cpp
@@ -251,7 +251,7 @@ FReply SDMXPixelMappingDMXLibraryView::OnAddDMXLibraryButtonClicked()
 	if (OtherFixtureGroupComponents.IsEmpty())
 	{
 		// If there's no group, add one that scales the current texture to the active renderer component
-		if (Toolkit->CanSizeSelectedComponentToTexture(), TEXT("Unexpected cannot size single fixture group component to texture"))
+		if (Toolkit->CanSizeSelectedComponentToTexture())
 		{
 			Toolkit->SizeSelectedComponentToTexture(true);
 		}
