DMX - Add post process options options to PixelMapping

diff --git a/Engine/Plugins/VirtualProduction/DMX/DMXPixelMapping/Source/DMXPixelMappingRenderer/Private/DMXPixelMappingPostProcessProxy.cpp b/Engine/Plugins/VirtualProduction/DMX/DMXPixelMapping/Source/DMXPixelMappingRenderer/Private/DMXPixelMappingPostProcessProxy.cpp
new file mode 100644
--- /dev/null
+++ b/Engine/Plugins/VirtualProduction/DMX/DMXPixelMapping/Source/DMXPixelMappingRenderer/Private/DMXPixelMappingPostProcessProxy.cpp
@@ -0,0 +1,249 @@
+// Copyright Epic Games, Inc. All Rights Reserved.
+
+#include "DMXPixelMappingPostProcessProxy.h"
+
+#include "CanvasTypes.h"
+#include "CommonRenderResources.h"
+#include "GlobalShader.h"
+#include "RHIStaticStates.h"
+#include "ScreenRendering.h"
+#include "TextureResource.h"
+#include "Engine/Canvas.h"
+#include "Engine/Engine.h"
+#include "Engine/Texture.h"
+#include "Engine/TextureRenderTarget2D.h"
+#include "Materials/MaterialInstanceDynamic.h"
+#include "Materials/MaterialInterface.h"
+#include "Materials/MaterialRenderProxy.h"
+#include "Modules/ModuleManager.h"
+
+
+namespace UE::DMXPixelMapping::Renderer::Private
+{
+	FDMXPixelMappingPostProcessCanvas::FDMXPixelMappingPostProcessCanvas()
+	{
+		Canvas = NewObject<UCanvas>();
+	}
+
+	void FDMXPixelMappingPostProcessCanvas::DrawMaterialToRenderTarget(UTextureRenderTarget2D* TextureRenderTarget, UMaterialInterface* Material)
+	{
+		if (!Material || !TextureRenderTarget || !TextureRenderTarget->GetResource())
+		{
+			return;
+		}
+
+		// This is a user-facing function, so we'd rather make sure that shaders are ready by the time we render, in order to ensure we don't draw with a fallback material
+		Material->EnsureIsComplete();
+		FTextureRenderTargetResource* RenderTargetResource = TextureRenderTarget->GameThread_GetRenderTargetResource();
+		if (!RenderTargetResource)
+		{
+			return;
+		}
+
+		FCanvas RenderCanvas(
+			RenderTargetResource,
+			nullptr,
+			GEngine->GetWorld(),
+			GMaxRHIFeatureLevel);
+
+		Canvas->Init(TextureRenderTarget->SizeX, TextureRenderTarget->SizeY, nullptr, &RenderCanvas);
+
+		{
+			ENQUEUE_RENDER_COMMAND(FlushDeferredResourceUpdateCommand)(
+				[RenderTargetResource](FRHICommandListImmediate& RHICmdList)
+				{
+					RenderTargetResource->FlushDeferredResourceUpdate(RHICmdList);
+				});
+
+			Canvas->K2_DrawMaterial(Material, FVector2D(0, 0), FVector2D(TextureRenderTarget->SizeX, TextureRenderTarget->SizeY), FVector2D(0, 0));
+
+			RenderCanvas.Flush_GameThread();
+			Canvas->Canvas = nullptr;
+
+			// UpdateResourceImmediate must be called here to ensure mips are generated.
+			TextureRenderTarget->UpdateResourceImmediate(false);
+
+			ENQUEUE_RENDER_COMMAND(ResetSceneTextureExtentHistory)(
+				[RenderTargetResource](FRHICommandListImmediate& RHICmdList)
+				{
+					RenderTargetResource->ResetSceneTextureExtentsHistory();
+				});
+		}
+	}
+
+	void FDMXPixelMappingPostProcessCanvas::AddReferencedObjects(FReferenceCollector& Collector)
+	{
+		Collector.AddReferencedObject(Canvas);
+	}
+
+
+	void FDMXPixelMappingPostProccessProxy::Render(UTexture* InInputTexture, const FDMXPixelMappingInputTextureRenderingParameters& InParams)
+	{
+		InputTexture = InInputTexture;
+		PostProcessMID = InParams.PostProcessMID;
+		UpdateRenderTargets(InParams.OutputSize, InParams.NumDownsamplePasses);
+
+		if (!CanRender())
+		{
+			return;
+		}
+
+		if (DownsampleRenderTargets.IsEmpty() && PostProcessMID)
+		{
+			// Render without downsampling
+			PostProcessMID->SetTextureParameterValue(InParams.PostProcessMaterialInputTextureParameterName, InputTexture);
+			PostProcessMID->SetScalarParameterValue(InParams.BlurDistanceParameterName, InParams.BlurDistance);
+			PostProcessCanvas.DrawMaterialToRenderTarget(OutputRenderTarget, PostProcessMID);
+		}
+		else
+		{
+			// Render with downsampling. 
+			// Note, if the last downsample render target is of minimal size, further renderer targets are not created.
+			// Hence use the number of downsample render targets instead of InParams.NumDownsamplePasses.
+			const int32 NumDownsamplePasses = DownsampleRenderTargets.Num();
+			for (int32 DownsamplePass = 0; DownsamplePass < NumDownsamplePasses; DownsamplePass++)
+			{
+				UTexture* SourceTexture = DownsamplePass == 0 ? InputTexture : DownsampleRenderTargets[DownsamplePass - 1];
+				UTextureRenderTarget2D* DownsampleRenderTarget = DownsampleRenderTargets[DownsamplePass].Get();
+
+				const bool bApplyPostProcessMaterial = PostProcessMID && (InParams.bApplyPostProcessMaterialEachDownsamplePass || DownsampleRenderTarget == DownsampleRenderTargets.Last());
+				if (bApplyPostProcessMaterial)
+				{
+					PostProcessMID->SetTextureParameterValue(InParams.PostProcessMaterialInputTextureParameterName, SourceTexture);
+					PostProcessMID->SetScalarParameterValue(InParams.BlurDistanceParameterName, InParams.BlurDistance);
+					PostProcessCanvas.DrawMaterialToRenderTarget(DownsampleRenderTarget, PostProcessMID);
+				}
+				else
+				{
+					RenderTextureToTarget(SourceTexture, DownsampleRenderTarget);
+				}
+			}
+
+			// Upscale
+			UTexture* LastSourceTexture = DownsampleRenderTargets.Last();
+			RenderTextureToTarget(LastSourceTexture, OutputRenderTarget);
+		}
+	}
+
+	UTexture* FDMXPixelMappingPostProccessProxy::GetRenderedTextureGameThread() const
+	{
+		return CanRender() ? OutputRenderTarget : InputTexture;
+	}
+
+	void FDMXPixelMappingPostProccessProxy::AddReferencedObjects(FReferenceCollector& Collector)
+	{
+		Collector.AddReferencedObject(InputTexture);
+		Collector.AddReferencedObject(PostProcessMID);
+		Collector.AddReferencedObjects(DownsampleRenderTargets);
+		Collector.AddReferencedObject(OutputRenderTarget);
+	}
+
+	bool FDMXPixelMappingPostProccessProxy::CanRender() const
+	{
+		return InputTexture && OutputRenderTarget && FApp::CanEverRender();
+	}
+
+	void FDMXPixelMappingPostProccessProxy::UpdateRenderTargets(const FVector2D& OutputSize, int32 NumDownsamplePasses)
+	{
+		const bool bNeedsAnyRenderTargets = InputTexture && (NumDownsamplePasses > 0 || PostProcessMID);
+		if (!bNeedsAnyRenderTargets)
+		{
+			OutputRenderTarget = nullptr;
+			return;
+		}
+
+		bool bInvalidRendertargets = !OutputRenderTarget || DownsampleRenderTargets.IsEmpty() || DownsampleRenderTargets.Num() != NumDownsamplePasses;
+		if (!bInvalidRendertargets)
+		{
+			if (OutputRenderTarget->GetSurfaceWidth() != OutputSize.X ||
+				OutputRenderTarget->GetSurfaceHeight() != OutputSize.Y)
+			{
+				// Input texture was resized 
+				bInvalidRendertargets = true;
+			}
+		}
+
+		if (bInvalidRendertargets)
+		{
+			const FVector2D InputSize = FVector2D(InputTexture->GetSurfaceHeight(), InputTexture->GetSurfaceWidth());
+
+			// Create downsample render targets
+			DownsampleRenderTargets.Reset();
+			FVector2D DownsampleSize = InputSize;
+			for (int32 DownsamplePass = 0; DownsamplePass < NumDownsamplePasses; DownsamplePass++)
+			{
+				DownsampleSize /= 2.0;
+				if (DownsampleSize.X <= 1.0 || DownsampleSize.Y <= 1.0)
+				{
+					break;
+				}
+				UTextureRenderTarget2D* DownsampleRenderTarget = NewObject<UTextureRenderTarget2D>();
+				DownsampleRenderTarget->ClearColor = FLinearColor::Black;
+				DownsampleRenderTarget->InitAutoFormat(DownsampleSize.X, DownsampleSize.Y);
+				DownsampleRenderTarget->UpdateResourceImmediate();
+				DownsampleRenderTargets.Add(DownsampleRenderTarget);
+			}
+
+			// Create output render target
+			OutputRenderTarget = NewObject<UTextureRenderTarget2D>();
+			OutputRenderTarget->ClearColor = FLinearColor::Black;
+			OutputRenderTarget->InitAutoFormat(OutputSize.X, OutputSize.Y);
+			OutputRenderTarget->UpdateResourceImmediate();
+		}
+	}
+
+	void FDMXPixelMappingPostProccessProxy::RenderTextureToTarget(UTexture* Texture, UTextureRenderTarget2D* RenderTarget) const
+	{
+		IRendererModule* RendererModule = FModuleManager::GetModulePtr<IRendererModule>("Renderer");
+		ENQUEUE_RENDER_COMMAND(PixelMappingPostProccessRenderToTargetPass)(
+			[RendererModule, Texture, RenderTarget, this](FRHICommandListImmediate& RHICmdList)
+			{
+				if (!Texture || !Texture->GetResource())
+				{
+					return;
+				}
+
+				const FTextureRHIRef SourceTextureRHI = Texture->GetResource()->GetTexture2DRHI();
+				const FTextureRHIRef TargetTextureRHI = RenderTarget->GetResource()->GetTexture2DRHI();
+
+				FRHIRenderPassInfo RPInfo(TargetTextureRHI, MakeRenderTargetActions(ERenderTargetLoadAction::ENoAction, ERenderTargetStoreAction::EStore));
+
+				RHICmdList.Transition(FRHITransitionInfo(TargetTextureRHI, ERHIAccess::Unknown, ERHIAccess::RTV));
+				RHICmdList.BeginRenderPass(RPInfo, TEXT("PixelMappingPostProcessProxy_RenderToTarget"));
+				{
+					RHICmdList.SetViewport(0, 0, 0.0f, (float)TargetTextureRHI->GetSizeX(), (float)TargetTextureRHI->GetSizeY(), 1.0f);
+
+					FGraphicsPipelineStateInitializer GraphicsPSOInit;
+					RHICmdList.ApplyCachedRenderTargets(GraphicsPSOInit);
+					GraphicsPSOInit.BlendState = TStaticBlendState<>::GetRHI();
+					GraphicsPSOInit.RasterizerState = TStaticRasterizerState<>::GetRHI();
+					GraphicsPSOInit.DepthStencilState = TStaticDepthStencilState<false, CF_Always>::GetRHI();
+
+					FGlobalShaderMap* ShaderMap = GetGlobalShaderMap(GMaxRHIFeatureLevel);
+					TShaderMapRef<FScreenVS> VertexShader(ShaderMap);
+					TShaderMapRef<FScreenPS> PixelShader(ShaderMap);
+
+					GraphicsPSOInit.BoundShaderState.VertexDeclarationRHI = GFilterVertexDeclaration.VertexDeclarationRHI;
+					GraphicsPSOInit.BoundShaderState.VertexShaderRHI = VertexShader.GetVertexShader();
+					GraphicsPSOInit.BoundShaderState.PixelShaderRHI = PixelShader.GetPixelShader();
+					GraphicsPSOInit.PrimitiveType = PT_TriangleList;
+
+					SetGraphicsPipelineState(RHICmdList, GraphicsPSOInit, 0);
+
+					SetShaderParametersLegacyPS(RHICmdList, PixelShader, TStaticSamplerState<SF_Bilinear>::GetRHI(), SourceTextureRHI);
+
+					RendererModule->DrawRectangle(RHICmdList,
+						0, 0,									// Dest X, Y
+						(float)TargetTextureRHI->GetSizeX(),	// Dest Width
+						(float)TargetTextureRHI->GetSizeY(),	// Dest Height
+						0, 0,									// Source U, V
+						1, 1,									// Source USize, VSize
+						TargetTextureRHI->GetSizeXY(),			// Target buffer size
+						FIntPoint(1, 1),						// Source texture size
+						VertexShader,
+						EDRF_Default);
+				}
+				RHICmdList.EndRenderPass();
+			});
+	}
+}
diff --git a/Engine/Plugins/VirtualProduction/DMX/DMXPixelMapping/Source/DMXPixelMappingRenderer/Private/DMXPixelMappingPostProcessProxy.h b/Engine/Plugins/VirtualProduction/DMX/DMXPixelMapping/Source/DMXPixelMappingRenderer/Private/DMXPixelMappingPostProcessProxy.h
new file mode 100644
--- /dev/null
+++ b/Engine/Plugins/VirtualProduction/DMX/DMXPixelMapping/Source/DMXPixelMappingRenderer/Private/DMXPixelMappingPostProcessProxy.h
@@ -0,0 +1,94 @@
+// Copyright Epic Games, Inc. All Rights Reserved.
+
+#pragma once
+
+#include "IDMXPixelMappingRenderer.h"
+
+#include "UObject/GCObject.h"
+#include "UObject/ObjectPtr.h"
+
+class UCanvas;
+class UMaterialInstanceDynamic;
+class UMaterialInterface;
+class UTexture;
+class UTextureRenderTarget2D;
+
+
+namespace UE::DMXPixelMapping::Renderer::Private
+{
+	/** Canvas implementation for pixelmapping */
+	class FDMXPixelMappingPostProcessCanvas
+		: public FGCObject
+	{
+	public:
+		/** Constructor */
+		FDMXPixelMappingPostProcessCanvas();
+
+		/** Draws a material to a render target */
+		void DrawMaterialToRenderTarget(UTextureRenderTarget2D* TextureRenderTarget, UMaterialInterface* Material);
+
+	protected:
+		//~ Begin FGCObject interface
+		virtual void AddReferencedObjects(FReferenceCollector& Collector) override;
+		virtual FString GetReferencerName() const override
+		{
+			return TEXT("FDMXPixelMappingPostProccessProxy");
+		}
+		//~ End FGCObject interface
+
+	private:
+		/** Canvas to be used for post processing */
+		TObjectPtr<UCanvas> Canvas;
+	};
+
+
+	/** Proxy to post process the texture used in Pixel Mapping */
+	class FDMXPixelMappingPostProccessProxy
+		: public FGCObject
+	{
+	public:
+		/** Renders the input texture  */
+		void Render(UTexture* InInputTexture, const FDMXPixelMappingInputTextureRenderingParameters& InParams);
+
+		/** Returns the rendered texture, or nullptr if no texture was rendered */
+		UTexture* GetRenderedTextureGameThread() const;
+
+	protected:
+		//~ Begin FGCObject interface
+		virtual void AddReferencedObjects(FReferenceCollector& Collector) override;
+		virtual FString GetReferencerName() const override
+		{
+			return TEXT("FDMXPixelMappingPostProccessProxy");
+		}
+		//~ End FGCObject interface
+
+	private:
+		/** Returns true if rendering is needed */
+		bool CanRender() const;
+
+		/** Updates render targets */
+		void UpdateRenderTargets(const FVector2D& OutputSize, int32 NumDownsamplePasses);
+
+		/** Renders the Input Texture to the Output render target*/
+		void RenderTextureToTarget(UTexture* Texture, UTextureRenderTarget2D* RenderTarget) const;
+
+		/** The number of times rendering is required, either to downsample or applying the post process material */
+		int32 NumRenderPasses = 0;
+
+		/** Weak ref to the input texture */
+		TObjectPtr<UTexture> InputTexture;
+
+		/** Post process material instance dynamic to be applied to the input texture */
+		TObjectPtr<UMaterialInstanceDynamic> PostProcessMID;
+
+		/** Downsample render targets */
+		TArray<TObjectPtr<UTextureRenderTarget2D>> DownsampleRenderTargets;
+
+		/** The render target that is output. Only valid if rendering is needed */
+		TObjectPtr<UTextureRenderTarget2D> OutputRenderTarget;
+
+		/** Post process canvas used to draw materials */
+		FDMXPixelMappingPostProcessCanvas PostProcessCanvas;
+	};
+
+}
diff --git a/Engine/Plugins/VirtualProduction/DMX/DMXPixelMapping/Source/DMXPixelMappingRenderer/Private/DMXPixelMappingRenderer.cpp b/Engine/Plugins/VirtualProduction/DMX/DMXPixelMapping/Source/DMXPixelMappingRenderer/Private/DMXPixelMappingRenderer.cpp
--- a/Engine/Plugins/VirtualProduction/DMX/DMXPixelMapping/Source/DMXPixelMappingRenderer/Private/DMXPixelMappingRenderer.cpp
+++ b/Engine/Plugins/VirtualProduction/DMX/DMXPixelMapping/Source/DMXPixelMappingRenderer/Private/DMXPixelMappingRenderer.cpp
@@ -17,6 +17,12 @@
 #include "TextureResource.h"
 #include "Widgets/Images/SImage.h"
 #include "Widgets/SOverlay.h"
+#include "TextureResource.h"
+#include "Rendering/Texture2DResource.h"
+#include "Kismet/KismetRenderingLibrary.h"
+#include "CanvasTypes.h"
+#include "Engine/Canvas.h"
+#include "Materials/MaterialInstanceDynamic.h"
 
 namespace DMXPixelMappingRenderer
 {
@@ -83,10 +89,10 @@ public:
 	static bool ShouldCompilePermutation(const FGlobalShaderPermutationParameters& Parameters) { return true; }
 };
 
-
 IMPLEMENT_GLOBAL_SHADER(FDMXPixelMappingRendererVS, "/Plugin/DMXPixelMapping/Private/DMXPixelMapping.usf", "DMXPixelMappingVS", SF_Vertex);
 IMPLEMENT_GLOBAL_SHADER(FDMXPixelMappingRendererPS, "/Plugin/DMXPixelMapping/Private/DMXPixelMapping.usf", "DMXPixelMappingPS", SF_Pixel);
 
+
 FDMXPixelMappingRenderer::FDMXPixelMappingRenderer()
 {
 	static const FName RendererModuleName("Renderer");
@@ -116,6 +122,16 @@ FDMXPixelMappingRenderer::FDMXPixelMappingRenderer()
 	}
 }
 
+void FDMXPixelMappingRenderer::PostProcessTexture(UTexture* InputTexture, const UE::DMXPixelMapping::FDMXPixelMappingInputTextureRenderingParameters& Params) const
+{
+	PostProcessProxy.Render(InputTexture, Params);
+}
+
+UTexture* FDMXPixelMappingRenderer::GetPostProcessedTexture() const
+{
+	return PostProcessProxy.GetRenderedTextureGameThread();
+}
+
 void FDMXPixelMappingRenderer::DownsampleRender(
 	const FTextureResource* InputTexture,
 	const FTextureResource* DstTexture,
@@ -346,7 +362,7 @@ void FDMXPixelMappingRenderer::RenderMaterial(UTextureRenderTarget2D* InRenderTa
 
 		static const float DeltaTime = 0.f;
 		MaterialWidgetRenderer->DrawWidget(InRenderTarget, Widget, TextureSize, DeltaTime);
-
+		
 		// Reset material after drawing
 		UIMaterialBrush->SetMaterial(nullptr);
 	}
@@ -369,6 +385,7 @@ void FDMXPixelMappingRenderer::RenderWidget(UTextureRenderTarget2D* InRenderTarg
 
 	UMGRenderer->DrawWidget(InRenderTarget, InUserWidget->TakeWidget(), TextureSize, DeltaTime);
 }
+
 void FDMXPixelMappingRenderer::RenderTextureToRectangle(const FTextureResource* InTextureResource, const FTexture2DRHIRef InRenderTargetTexture, FVector2D InSize, bool bSRGBSource) const
 {
 	check(IsInGameThread());
@@ -381,6 +398,11 @@ void FDMXPixelMappingRenderer::RenderTextureToRectangle(const FTextureResource*
 		bool bSRGBSource;
 	};
 
+	if (!InTextureResource)
+	{
+		return;
+	}
+
 	FRenderContext RenderContext
 	{
 		InTextureResource,
diff --git a/Engine/Plugins/VirtualProduction/DMX/DMXPixelMapping/Source/DMXPixelMappingRenderer/Private/DMXPixelMappingRenderer.h b/Engine/Plugins/VirtualProduction/DMX/DMXPixelMapping/Source/DMXPixelMappingRenderer/Private/DMXPixelMappingRenderer.h
--- a/Engine/Plugins/VirtualProduction/DMX/DMXPixelMapping/Source/DMXPixelMappingRenderer/Private/DMXPixelMappingRenderer.h
+++ b/Engine/Plugins/VirtualProduction/DMX/DMXPixelMapping/Source/DMXPixelMappingRenderer/Private/DMXPixelMappingRenderer.h
@@ -4,11 +4,15 @@
 
 #include "IDMXPixelMappingRenderer.h"
 
+#include "DMXPixelMappingPostProcessProxy.h"
+
 class IRendererModule;
 class FWidgetRenderer;
 
 struct FSlateMaterialBrush;
 
+
+
 /**
  * Implementation of Pixel Mapping Renderer
  */
@@ -20,9 +24,12 @@ public:
 	FDMXPixelMappingRenderer();
 
 	//~ Begin IDMXPixelMappingRenderer implementation
-	virtual void RenderMaterial(UTextureRenderTarget2D* InRenderTarget, UMaterialInterface* InMaterialInterface) const override;
-	virtual void RenderWidget(UTextureRenderTarget2D* InRenderTarget, UUserWidget* InUserWidget) const override;
 
+	virtual void PostProcessTexture(UTexture* InputTexture, const UE::DMXPixelMapping::FDMXPixelMappingInputTextureRenderingParameters& Params) const override;
+
+	virtual UTexture* GetPostProcessedTexture() const override;
+
+	
 	virtual void DownsampleRender(
 		const FTextureResource* InputTexture,
 		const FTextureResource* DstTexture,
@@ -31,6 +38,8 @@ public:
 		DownsampleReadCallback InCallback
 	) const override;
 
+	virtual void RenderMaterial(UTextureRenderTarget2D* InRenderTarget, UMaterialInterface* InMaterialInterface) const override;
+	virtual void RenderWidget(UTextureRenderTarget2D* InRenderTarget, UUserWidget* InUserWidget) const override;
 	virtual void RenderTextureToRectangle(const FTextureResource* InTextureResource, const FTexture2DRHIRef InRenderTargetTexture, FVector2D InSize, bool bSRGBSource) const override;
 
 #if WITH_EDITOR
@@ -39,6 +48,9 @@ public:
 	//~ End IDMXPixelMappingRenderer implementation
 
 private:
+	/** Proxy to post process the input texture */
+	mutable UE::DMXPixelMapping::Renderer::Private::FDMXPixelMappingPostProccessProxy PostProcessProxy;
+
 	/** Brush for Material widget renderer */
 	TSharedPtr<FSlateMaterialBrush> UIMaterialBrush;
 
diff --git a/Engine/Plugins/VirtualProduction/DMX/DMXPixelMapping/Source/DMXPixelMappingRenderer/Public/IDMXPixelMappingRenderer.h b/Engine/Plugins/VirtualProduction/DMX/DMXPixelMapping/Source/DMXPixelMappingRenderer/Public/IDMXPixelMappingRenderer.h
--- a/Engine/Plugins/VirtualProduction/DMX/DMXPixelMapping/Source/DMXPixelMappingRenderer/Public/IDMXPixelMappingRenderer.h
+++ b/Engine/Plugins/VirtualProduction/DMX/DMXPixelMapping/Source/DMXPixelMappingRenderer/Public/IDMXPixelMappingRenderer.h
@@ -6,12 +6,16 @@
 #include "RHI.h"
 #include "RHIResources.h"
 
+enum class EDMXPixelBlendingQuality : uint8;
+
 class FTextureResource;
 class FTextureRenderTargetResource;
 class UTextureRenderTarget2D;
+class UMaterialInstanceDynamic;
 class UMaterialInterface;
+class UTexture;
 class UUserWidget;
-enum class EDMXPixelBlendingQuality : uint8;
+
 
 /**
  * Used in shader permutation for determining number of samples to use in texture blending.
@@ -26,6 +30,39 @@ enum class EDMXPixelShaderBlendingQuality : uint8
 	MAX
 };
 
+namespace UE::DMXPixelMapping
+{
+	/** Parameters for the Input Texture Renderer */
+	struct FDMXPixelMappingInputTextureRenderingParameters
+	{
+		/** Number of times a texture is downsampled. E.g. when texture size is 512px and is downsampled 3 times, its resulting size is 64px */
+		int32 NumDownsamplePasses = 0;
+
+		/** The post process material. If null, no post process material is applied */
+		UMaterialInstanceDynamic* PostProcessMID = nullptr;
+
+		/** The input texture parameter name of the post process material */
+		FName PostProcessMaterialInputTextureParameterName;
+
+		/** The input texture parameter name of the post process material */
+		FName BlurDistanceParameterName;
+
+		/** The blur distance of the post process material */
+		float BlurDistance = .2f;
+
+		/**
+		 * If true, applies post process material each downsample pass.
+		 * If false applies the post process material once after the last downsample pass, or direct if the input is not downsampled.
+		 * Only applicable if a post process material is set.
+		 */
+		bool bApplyPostProcessMaterialEachDownsamplePass = true;
+
+		/** Size of the rendered texture */
+		FVector2D OutputSize{ 1.f, 1.f };
+	};
+
+}
+
 /**
  * Downsample pixel preview rendering params.
  * Using for pixel rendering setting in preview
@@ -95,6 +132,20 @@ struct FDMXPixelMappingDownsamplePixelParamsV2
 	bool bStaticCalculateUV;
 };
 
+struct FDMXPixelMappingRenderTextureParams
+{
+	int32 DownsampleTexture = 8;
+
+	int32 NumDownSamplePasses = 1;
+	float Distance = 0.2;
+	int32 DistanceSteps = 1;
+	int32 RadialSteps = 1;
+	float RadialOffset = 1;
+	int32 KernelPower = 5;
+};
+
+
+
 
 /**
  * The public interface of the Pixel Mapping renderer instance interface.
@@ -108,9 +159,20 @@ public:
 public:
 	/** Virtual destructor */
 	virtual ~IDMXPixelMappingRenderer() = default;
+		
+	/**
+	 * Blurs input texture onto desination texture
+	 *
+	 * @param InputTexture						The input texture that is being processed
+	 * @param Params							Parameters for post processing.
+	 */
+	virtual void PostProcessTexture(UTexture* InputTexture, const UE::DMXPixelMapping::FDMXPixelMappingInputTextureRenderingParameters& Params) const = 0;
+
+	/** Gets the post processed texture. May return nullptr while the texture is not rendered yet. */
+	virtual UTexture* GetPostProcessedTexture() const = 0;
 
 	/**
-	 * Downsample and Draw input texture to Destination texture.
+	 * Pixelmapping specific, downsample and draw input texture to destination texture.
 	 *
 	 * @param InputTexture					Rendering resource of input texture
 	 * @param DstTexture					Rendering resource of RenderTarget texture
diff --git a/Engine/Plugins/VirtualProduction/DMX/DMXPixelMapping/Source/DMXPixelMappingRuntime/Private/Components/DMXPixelMappingRendererComponent.cpp b/Engine/Plugins/VirtualProduction/DMX/DMXPixelMapping/Source/DMXPixelMappingRuntime/Private/Components/DMXPixelMappingRendererComponent.cpp
--- a/Engine/Plugins/VirtualProduction/DMX/DMXPixelMapping/Source/DMXPixelMappingRuntime/Private/Components/DMXPixelMappingRendererComponent.cpp
+++ b/Engine/Plugins/VirtualProduction/DMX/DMXPixelMapping/Source/DMXPixelMappingRuntime/Private/Components/DMXPixelMappingRendererComponent.cpp
@@ -4,18 +4,20 @@
 
 #include "DMXPixelMappingMainStreamObjectVersion.h"
 #include "DMXPixelMappingTypes.h"
+#include "DMXStats.h"
 #include "IDMXPixelMappingRenderer.h"
 #include "IDMXPixelMappingRendererModule.h"
-#include "DMXStats.h"
 #include "Components/DMXPixelMappingFixtureGroupComponent.h"
 #include "Components/DMXPixelMappingFixtureGroupItemComponent.h"
 #include "Components/DMXPixelMappingMatrixComponent.h"
 #include "Components/DMXPixelMappingRootComponent.h"
 #include "Components/DMXPixelMappingScreenComponent.h"
+#include "Rendering/RenderInputTextureProxy.h"
 
 #include "Async/Async.h"
 #include "Blueprint/UserWidget.h"
 #include "Engine/TextureRenderTarget2D.h"
+#include "Materials/MaterialInstanceDynamic.h"
 #include "Modulators/DMXModulator.h"
 #include "TextureResource.h"
 #include "UObject/ConstructorHelpers.h"
@@ -31,12 +33,11 @@
 DECLARE_CYCLE_STAT(TEXT("PixelMapping Render"), STAT_DMXPixelMappingRender, STATGROUP_DMX);
 DECLARE_CYCLE_STAT(TEXT("PixelMapping Send DMX Worker"), STAT_DMXPixelMappingSendDMX, STATGROUP_DMX);
 
+
 const FIntPoint UDMXPixelMappingRendererComponent::MaxDownsampleBufferTargetSize = FIntPoint(4096);
 const FLinearColor UDMXPixelMappingRendererComponent::ClearTextureColor = FLinearColor::Black;
 
 UDMXPixelMappingRendererComponent::UDMXPixelMappingRendererComponent()
-	: DownsampleBufferTarget(nullptr)
-	, DownsamplePixelCount(0)
 {
 	SetSize(FVector2D(100.f, 100.f));
 
@@ -62,17 +63,6 @@ UDMXPixelMappingRendererComponent::UDMXPixelMappingRendererComponent()
 #endif
 }
 
-UDMXPixelMappingRendererComponent::~UDMXPixelMappingRendererComponent()
-{
-#if WITH_EDITOR
-	if (OnChangeLevelHandle.IsValid())
-	{
-		FLevelEditorModule& LevelEditorModule = FModuleManager::LoadModuleChecked<FLevelEditorModule>(TEXT("LevelEditor"));
-		LevelEditorModule.OnMapChanged().Remove(OnChangeLevelHandle);
-	}
-#endif
-}
-
 const FName& UDMXPixelMappingRendererComponent::GetNamePrefix()
 {
 	static FName NamePrefix = TEXT("Renderer");
@@ -111,6 +101,16 @@ void UDMXPixelMappingRendererComponent::Serialize(FArchive& Ar)
 void UDMXPixelMappingRendererComponent::PostInitProperties()
 {
 	Super::PostInitProperties();
+
+	if (!IsTemplate())
+	{
+		Initialize();
+	}
+}
+
+void UDMXPixelMappingRendererComponent::PostLoad()
+{
+	Super::PostLoad();
 	
 	if (!IsTemplate())
 	{
@@ -124,6 +124,7 @@ void UDMXPixelMappingRendererComponent::PostEditChangeChainProperty(FPropertyCha
 	// Call the parent at the first place
 	Super::PostEditChangeChainProperty(PropertyChangedChainEvent);
 
+	const FName PropertyName = PropertyChangedChainEvent.GetPropertyName();
 	if (PropertyChangedChainEvent.GetPropertyName() == UDMXPixelMappingOutputComponent::GetSizeXPropertyName() ||
 		PropertyChangedChainEvent.GetPropertyName() == UDMXPixelMappingOutputComponent::GetSizeYPropertyName())
 	{
@@ -140,8 +141,6 @@ void UDMXPixelMappingRendererComponent::PostEditChangeChainProperty(FPropertyCha
 
 			UE_LOG(LogDMXPixelMappingRuntime, Warning, TEXT("Pixel mapping textures are limited to engine's max texture dimension %dx%d"), MaxTextureDimensions, MaxTextureDimensions);
 		}
-
-		ResizeMaterialRenderTarget(GetSize().X, GetSize().Y);
 	} 
 	else if (PropertyChangedChainEvent.GetPropertyName() == GET_MEMBER_NAME_CHECKED(UDMXPixelMappingRendererComponent, RendererType))
 	{
@@ -299,16 +298,6 @@ TSharedRef<SWidget> UDMXPixelMappingRendererComponent::TakeWidget()
 }
 #endif // WITH_EDITOR
 
-#if WITH_EDITOR
-void UDMXPixelMappingRendererComponent::OnMapChanged(UWorld* InWorld, EMapChangeType MapChangeType)
-{
-	if (UserWidget != nullptr)
-	{
-		UserWidget = nullptr;
-	}
-}
-#endif // WITH_EDITOR
-
 #if WITH_EDITOR
 void UDMXPixelMappingRendererComponent::ResizePreviewRenderTarget(uint32 InSizeX, uint32 InSizeY)
 {
@@ -356,23 +345,21 @@ void UDMXPixelMappingRendererComponent::Render()
 		return;
 	}
 
-	// 1. Get downsample input texture
+	// 1. Render the input texture
+	RendererInputTexture();
 	UTexture* DownsampleInputTexture = GetRendererInputTexture();
 	if (!DownsampleInputTexture)
 	{
 		return;
 	}
 
-	// 2. Render the input texture before downsample
-	RendererInputTexture();
-
-	// 3. Make sure there is the DownsampleBufferTarget exists and can size can hold all pixels
+	// 2. Make sure there is the DownsampleBufferTarget exists and can size can hold all pixels
 	CreateOrUpdateDownsampleBufferTarget();
 
-	// 4. reserve enough space for pixels params
+	// 3. reserve enough space for pixels params
 	DownsamplePixelParams.Reset(DownsamplePixelCount);
 
-	// 5. Loop through all child pixels in order to get pixels downsample params for rendering
+	// 4. Loop through all child pixels in order to get pixels downsample params for rendering
 	ForEachChild([&](UDMXPixelMappingBaseComponent* InComponent) {
 		if (UDMXPixelMappingOutputComponent* Component = Cast<UDMXPixelMappingOutputComponent>(InComponent))
 		{
@@ -380,14 +367,14 @@ void UDMXPixelMappingRendererComponent::Render()
 		}
 	}, false);
 
-	// 6. Make sure pixel count the same with pixel params set number
+	// 5. Make sure pixel count the same with pixel params set number
 	if (!ensure(DownsamplePixelParams.Num() == DownsamplePixelCount))
 	{
 		DownsamplePixelParams.Empty();
 		return;
 	}
 
-	// 7. Downsample all pixels
+	// 6. Downsample all pixels
 	PixelMappingRenderer->DownsampleRender(
 		DownsampleInputTexture->GetResource(),
 		DownsampleBufferTarget->GetResource(),
@@ -427,56 +414,19 @@ UWorld* UDMXPixelMappingRendererComponent::GetWorld() const
 	return World;
 }
 
-void UDMXPixelMappingRendererComponent::ResizeMaterialRenderTarget(int32 InSizeX, int32 InSizeY)
-{
-	check(InputRenderTarget);
-
-	if (InSizeX > 0 && InSizeY > 0)
-	{
-		if (InputRenderTarget->SizeX != InSizeY ||
-			InputRenderTarget->SizeY != InSizeX)
-		{
-			InputRenderTarget->ResizeTarget(InSizeX, InSizeY);
-			InputRenderTarget->UpdateResourceImmediate();
-		}
-	}
-}
-
 void UDMXPixelMappingRendererComponent::UpdateInputWidget(TSubclassOf<UUserWidget> InInputWidget)
 {
-	if (InInputWidget != nullptr && UserWidget != nullptr)
-	{
-		UserWidget->MarkAsGarbage();
-		UserWidget = nullptr;
-	}
-	else
-	{
-		UserWidget = CreateWidget(GetWorld(), InInputWidget);
-	}
+	UserWidget = CreateWidget(GetWorld(), InInputWidget);
 }
 
 void UDMXPixelMappingRendererComponent::Initialize()
 {
-	if (InputRenderTarget == nullptr)
-	{
-		InputRenderTarget = CreateRenderTarget(TEXT("InputRenderTarget"));
-	}	
-
-#if WITH_EDITOR
-	if (!OnChangeLevelHandle.IsValid())
-	{
-		FLevelEditorModule& LevelEditorModule = FModuleManager::LoadModuleChecked<FLevelEditorModule>(TEXT("LevelEditor"));
-		LevelEditorModule.OnMapChanged().AddUObject(this, &UDMXPixelMappingRendererComponent::OnMapChanged);
-	}
-#endif
-
 	if (InputTexture)
 	{
 		InputTexture->SetForceMipLevelsToBeResident(30.0f);
 		InputTexture->WaitForStreaming();
 	}
 
-
 	if (UserWidget == nullptr && InputWidget != nullptr)
 	{
 		UserWidget = CreateWidget(GetWorld(), InputWidget);
@@ -505,12 +455,14 @@ void UDMXPixelMappingRendererComponent::Initialize()
 				SetSize(NewSize);
 			}
 		}
-}
+	}
 	else
 	{
 		bLockInDesigner = false;
 	}
 #endif
+
+	UpdateInputTextureRenderProxy();
 }
 
 UTextureRenderTarget2D* UDMXPixelMappingRendererComponent::CreateRenderTarget(const FName& InBaseName)
@@ -521,26 +473,14 @@ UTextureRenderTarget2D* UDMXPixelMappingRendererComponent::CreateRenderTarget(co
 	constexpr bool bInForceLinearGamma = false;
 	RenderTarget->InitCustomFormat(GetSize().X, GetSize().Y, EPixelFormat::PF_B8G8R8A8, bInForceLinearGamma);
 
-	return  RenderTarget;
+	return RenderTarget;
 }
 
 void UDMXPixelMappingRendererComponent::RendererInputTexture()
 {
-	switch (RendererType)
+	if (RenderInputTextureProxy.IsValid())
 	{
-	case(EDMXPixelMappingRendererType::Texture):
-		// No rendering required
-		break;
-
-	case(EDMXPixelMappingRendererType::Material):
-		PixelMappingRenderer->RenderMaterial(InputRenderTarget, InputMaterial);
-		break;
-
-	case(EDMXPixelMappingRendererType::UMG):
-		PixelMappingRenderer->RenderWidget(InputRenderTarget, UserWidget);
-		break;
-	default:
-		checkNoEntry(); // Unknown renderer type
+		RenderInputTextureProxy->Render();
 	}
 
 #if WITH_EDITOR
@@ -560,28 +500,7 @@ void UDMXPixelMappingRendererComponent::RendererInputTexture()
 
 UTexture* UDMXPixelMappingRendererComponent::GetRendererInputTexture() const
 {
-	switch (RendererType)
-	{
-	case(EDMXPixelMappingRendererType::Texture):
-		return InputTexture;
-	case(EDMXPixelMappingRendererType::Material):
-		if (InputMaterial != nullptr)
-		{
-			return InputRenderTarget;
-		}
-		break;
-
-	case(EDMXPixelMappingRendererType::UMG):
-		if (InputWidget != nullptr)
-		{
-			return InputRenderTarget;
-		}
-		break;
-	default:
-		checkNoEntry();
-	}
-
-	return nullptr;
+	return RenderInputTextureProxy.IsValid() ? RenderInputTextureProxy->GetRenderedTexture() : nullptr;
 }
 
 bool UDMXPixelMappingRendererComponent::CanBeMovedTo(const UDMXPixelMappingBaseComponent* Component) const
@@ -589,32 +508,6 @@ bool UDMXPixelMappingRendererComponent::CanBeMovedTo(const UDMXPixelMappingBaseC
 	return Component && Component->IsA<UDMXPixelMappingRootComponent>();
 }
 
-int32 UDMXPixelMappingRendererComponent::GetTotalDownsamplePixelCount()
-{
-	FScopeLock ScopeLock(&DownsampleBufferCS);
-
-	// Reset pixel counter
-	DownsamplePixelCount = 0;
-
-	// Count all pixels
-	constexpr bool bIsRecursive = true;
-	ForEachChildOfClass<UDMXPixelMappingOutputComponent>([&](UDMXPixelMappingOutputComponent* InComponent)
-		{
-			// If that is screen component
-			if (UDMXPixelMappingScreenComponent* ScreenComponent = Cast<UDMXPixelMappingScreenComponent>(InComponent))
-			{
-				DownsamplePixelCount += (ScreenComponent->NumXCells * ScreenComponent->NumYCells);
-			}
-			// If that is single pixel component
-			else if (Cast<UDMXPixelMappingOutputDMXComponent>(InComponent))
-			{
-				DownsamplePixelCount++;
-			}
-		}, bIsRecursive);
-
-	return DownsamplePixelCount;
-}
-
 void UDMXPixelMappingRendererComponent::CreateOrUpdateDownsampleBufferTarget()
 {
 	// Create texture if it does not exists
@@ -672,20 +565,6 @@ int32 UDMXPixelMappingRendererComponent::GetDownsamplePixelNum()
 	return DownsamplePixelParams.Num();
 }
 
-bool UDMXPixelMappingRendererComponent::IsPixelRangeValid(const int32 InDownsamplePixelIndexStart, const int32 InDownsamplePixelIndexEnd) const
-{
-	FScopeLock ScopeLock(&DownsampleBufferCS);
-
-	if (InDownsamplePixelIndexEnd >= InDownsamplePixelIndexStart &&
-		DownsampleBuffer.IsValidIndex(InDownsamplePixelIndexStart) &&
-		DownsampleBuffer.IsValidIndex(InDownsamplePixelIndexEnd))
-	{
-		return true;
-	}
-
-	return false;
-}
-
 void UDMXPixelMappingRendererComponent::SetDownsampleBuffer(TArray<FLinearColor>&& InDownsampleBuffer, FIntRect InRect)
 {
 	check(IsInRenderingThread());
@@ -774,3 +653,71 @@ void UDMXPixelMappingRendererComponent::EmptyDownsampleBuffer()
 
 	DownsampleBuffer.Empty();
 }
+
+bool UDMXPixelMappingRendererComponent::IsPixelRangeValid(const int32 InDownsamplePixelIndexStart, const int32 InDownsamplePixelIndexEnd) const
+{
+	FScopeLock ScopeLock(&DownsampleBufferCS);
+
+	if (InDownsamplePixelIndexEnd >= InDownsamplePixelIndexStart &&
+		DownsampleBuffer.IsValidIndex(InDownsamplePixelIndexStart) &&
+		DownsampleBuffer.IsValidIndex(InDownsamplePixelIndexEnd))
+	{
+		return true;
+	}
+
+	return false;
+}
+
+int32 UDMXPixelMappingRendererComponent::GetTotalDownsamplePixelCount()
+{
+	FScopeLock ScopeLock(&DownsampleBufferCS);
+
+	// Reset pixel counter
+	DownsamplePixelCount = 0;
+
+	// Count all pixels
+	constexpr bool bIsRecursive = true;
+	ForEachChildOfClass<UDMXPixelMappingOutputComponent>([&](UDMXPixelMappingOutputComponent* InComponent)
+		{
+			// If that is screen component
+			if (UDMXPixelMappingScreenComponent* ScreenComponent = Cast<UDMXPixelMappingScreenComponent>(InComponent))
+			{
+				DownsamplePixelCount += (ScreenComponent->NumXCells * ScreenComponent->NumYCells);
+			}
+			// If that is single pixel component
+			else if (Cast<UDMXPixelMappingOutputDMXComponent>(InComponent))
+			{
+				DownsamplePixelCount++;
+			}
+		}, bIsRecursive);
+
+	return DownsamplePixelCount;
+}
+
+void UDMXPixelMappingRendererComponent::UpdateInputTextureRenderProxy()
+{
+	using namespace UE::DMXPixelMapping::Rendering::Private;
+
+	// Update params
+	RenderInputTextureParams.OutputSize = GetSize();
+	RenderInputTextureParams.PostProcessMID = PostProcessMaterial ? UMaterialInstanceDynamic::Create(PostProcessMaterial, GetTransientPackage()) : nullptr;
+
+	RenderInputTextureProxy.Reset();
+	switch (RendererType)
+	{
+	case(EDMXPixelMappingRendererType::Texture):
+		RenderInputTextureProxy = MakeShared<FRenderInputTextureProxy<UTexture>>(PixelMappingRenderer.ToSharedRef(), InputTexture.Get(), RenderInputTextureParams);
+		break;
+
+	case(EDMXPixelMappingRendererType::Material):
+		RenderInputTextureProxy = MakeShared<FRenderInputTextureProxy<UMaterialInterface>>(PixelMappingRenderer.ToSharedRef(), InputMaterial.Get(), RenderInputTextureParams, GetSize());
+		break;
+
+	case(EDMXPixelMappingRendererType::UMG):
+		RenderInputTextureProxy = MakeShared<FRenderInputTextureProxy<UUserWidget>>(PixelMappingRenderer.ToSharedRef(), InputWidget.GetDefaultObject(), RenderInputTextureParams, GetSize());
+		break;
+
+	default:
+		checkf(0, TEXT("Invalid Renderer Type in DMXPixelMappingRendererComponent"));
+	}
+}
diff --git a/Engine/Plugins/VirtualProduction/DMX/DMXPixelMapping/Source/DMXPixelMappingRuntime/Private/Rendering/RenderInputTextureProxy.cpp b/Engine/Plugins/VirtualProduction/DMX/DMXPixelMapping/Source/DMXPixelMappingRuntime/Private/Rendering/RenderInputTextureProxy.cpp
new file mode 100644
--- /dev/null
+++ b/Engine/Plugins/VirtualProduction/DMX/DMXPixelMapping/Source/DMXPixelMappingRuntime/Private/Rendering/RenderInputTextureProxy.cpp
@@ -0,0 +1,102 @@
+// Copyright Epic Games, Inc. All Rights Reserved.
+
+#include "RenderInputTextureProxy.h"
+
+#include "IDMXPixelMappingRenderer.h"
+#include "IDMXPixelMappingRendererModule.h"
+
+#include "Blueprint/UserWidget.h"
+#include "Engine/Texture.h"
+#include "Engine/TextureRenderTarget2D.h"
+#include "Materials/MaterialInterface.h"
+
+
+namespace UE::DMXPixelMapping::Rendering::Private
+{
+	static FDMXPixelMappingInputTextureRenderingParameters MakeRenderingParameters(const FDMXPixelMappingRenderInputTextureParameters& Parameters)
+	{
+		FDMXPixelMappingInputTextureRenderingParameters RenderingParams;
+		RenderingParams.NumDownsamplePasses = Parameters.NumDownSamplePasses;
+		RenderingParams.PostProcessMID = Parameters.PostProcessMID;
+		RenderingParams.BlurDistance = Parameters.BlurDistance;
+		RenderingParams.bApplyPostProcessMaterialEachDownsamplePass = Parameters.bApplyPostProcessMaterialEachDownsamplePass;
+		RenderingParams.PostProcessMaterialInputTextureParameterName = Parameters.InputTextureParameterName;
+		RenderingParams.BlurDistanceParameterName = Parameters.BlurDistanceParamtereName;
+		RenderingParams.OutputSize = Parameters.OutputSize;
+
+		return RenderingParams;
+	}
+
+	template<>
+	FRenderInputTextureProxy<UTexture>::FRenderInputTextureProxy(const TSharedRef<IDMXPixelMappingRenderer>& InPixelMappingRenderer, UTexture* InTexture, const FDMXPixelMappingRenderInputTextureParameters& InParameters)
+		: PixelMappingRenderer(InPixelMappingRenderer)
+		, WeakInputTexture(InTexture)
+		, Parameters(InParameters)
+	{}
+
+	template<>
+	FRenderInputTextureProxy<UMaterialInterface>::FRenderInputTextureProxy(const TSharedRef<IDMXPixelMappingRenderer>& InPixelMappingRenderer, UMaterialInterface* InMaterial, const FDMXPixelMappingRenderInputTextureParameters& InParameters, const FVector2D& InSize)
+		: PixelMappingRenderer(InPixelMappingRenderer)
+		, WeakInputTexture(InMaterial)
+		, Parameters(InParameters)
+	{
+		IntermediateRenderTarget = NewObject<UTextureRenderTarget2D>();
+		IntermediateRenderTarget->ClearColor = FLinearColor::Black;
+		IntermediateRenderTarget->InitAutoFormat(InSize.X, InSize.Y);
+		IntermediateRenderTarget->UpdateResourceImmediate();
+	}
+
+	template<>
+	FRenderInputTextureProxy<UUserWidget>::FRenderInputTextureProxy(const TSharedRef<IDMXPixelMappingRenderer>& InPixelMappingRenderer, UUserWidget* InUserWidget, const FDMXPixelMappingRenderInputTextureParameters& InParameters, const FVector2D& InSize)
+		: PixelMappingRenderer(InPixelMappingRenderer)
+		, WeakInputTexture(InUserWidget)
+		, Parameters(InParameters)
+	{
+		IntermediateRenderTarget = NewObject<UTextureRenderTarget2D>();
+		IntermediateRenderTarget->ClearColor = FLinearColor::Black;
+		IntermediateRenderTarget->InitAutoFormat(InSize.X, InSize.Y);
+		IntermediateRenderTarget->UpdateResourceImmediate();
+	}
+
+	template <typename TextureType>
+	void FRenderInputTextureProxy<TextureType>::UpdateParameters(const FDMXPixelMappingRenderInputTextureParameters& InParameters)
+	{
+		Parameters = InParameters;
+	}
+
+	template<>
+	void FRenderInputTextureProxy<UTexture>::Render()
+	{
+		const FDMXPixelMappingInputTextureRenderingParameters RenderingParams = MakeRenderingParameters(Parameters);
+
+		PixelMappingRenderer->PostProcessTexture(WeakInputTexture.Get(), RenderingParams);
+	}
+
+	template<>
+	void FRenderInputTextureProxy<UMaterialInterface>::Render()
+	{
+		check(IntermediateRenderTarget);
+		PixelMappingRenderer->RenderMaterial(IntermediateRenderTarget.Get(), CastChecked<UMaterialInterface>(WeakInputTexture.Get()));
+
+		const FDMXPixelMappingInputTextureRenderingParameters RenderingParams = MakeRenderingParameters(Parameters);
+
+		PixelMappingRenderer->PostProcessTexture(IntermediateRenderTarget.Get(), RenderingParams);
+	}
+
+	template<>
+	void FRenderInputTextureProxy<UUserWidget>::Render()
+	{
+		check(IntermediateRenderTarget);
+		PixelMappingRenderer->RenderWidget(IntermediateRenderTarget.Get(), CastChecked<UUserWidget>(WeakInputTexture.Get()));
+
+		const FDMXPixelMappingInputTextureRenderingParameters RenderingParams = MakeRenderingParameters(Parameters);
+
+		PixelMappingRenderer->PostProcessTexture(IntermediateRenderTarget.Get(), RenderingParams);
+	}
+
+	template <typename TextureType>
+	UTexture* FRenderInputTextureProxy<TextureType>::GetRenderedTexture() const
+	{
+		return PixelMappingRenderer->GetPostProcessedTexture();
+	}
+}
diff --git a/Engine/Plugins/VirtualProduction/DMX/DMXPixelMapping/Source/DMXPixelMappingRuntime/Private/Rendering/RenderInputTextureProxy.h b/Engine/Plugins/VirtualProduction/DMX/DMXPixelMapping/Source/DMXPixelMappingRuntime/Private/Rendering/RenderInputTextureProxy.h
new file mode 100644
--- /dev/null
+++ b/Engine/Plugins/VirtualProduction/DMX/DMXPixelMapping/Source/DMXPixelMappingRuntime/Private/Rendering/RenderInputTextureProxy.h
@@ -0,0 +1,45 @@
+// Copyright Epic Games, Inc. All Rights Reserved.
+
+#pragma once
+
+#include "Rendering/IDMXPixelMappingRenderInputTextureProxy.h"
+
+class IDMXPixelMappingRenderer;
+class IDMXPixelMappingRenderInputTextureProxy;
+
+class UTexture;
+class UTextureRenderTarget2D;
+
+
+namespace UE::DMXPixelMapping::Rendering::Private
+{
+	/** Base class for the render input texture proxy implementation. */
+	template <typename TextureType>
+	class FRenderInputTextureProxy
+		: public IDMXPixelMappingRenderInputTextureProxy
+	{
+	public:
+		FRenderInputTextureProxy(const TSharedRef<IDMXPixelMappingRenderer>& InPixelMappingRenderer, TextureType* InTexture, const FDMXPixelMappingRenderInputTextureParameters& InParameters);
+		FRenderInputTextureProxy(const TSharedRef<IDMXPixelMappingRenderer>& InPixelMappingRenderer, TextureType* InTexture, const FDMXPixelMappingRenderInputTextureParameters& InParameters, const FVector2D& InSize);
+
+	protected:
+		//~ Begin DMXPixelMappingRenderInputTextureProxy interface
+		virtual void UpdateParameters(const FDMXPixelMappingRenderInputTextureParameters& InParameters) override;
+		virtual void Render() override;
+		virtual UTexture* GetRenderedTexture() const override;
+		//~ End DMXPixelMappingRenderInputTextureProxy interface
+
+	private:
+		/** The renderer used with this proxy */
+		TSharedRef<IDMXPixelMappingRenderer> PixelMappingRenderer;
+
+		/** The Input Texture */
+		TWeakObjectPtr<TextureType> WeakInputTexture;
+
+		/** Intermediate render target for materials and user widgets */
+		TObjectPtr<UTextureRenderTarget2D> IntermediateRenderTarget;
+
+		/** Parameters for renderering */
+		FDMXPixelMappingRenderInputTextureParameters Parameters;
+	};
+}
diff --git a/Engine/Plugins/VirtualProduction/DMX/DMXPixelMapping/Source/DMXPixelMappingRuntime/Public/Components/DMXPixelMappingRendererComponent.h b/Engine/Plugins/VirtualProduction/DMX/DMXPixelMapping/Source/DMXPixelMappingRuntime/Public/Components/DMXPixelMappingRendererComponent.h
--- a/Engine/Plugins/VirtualProduction/DMX/DMXPixelMapping/Source/DMXPixelMappingRuntime/Public/Components/DMXPixelMappingRendererComponent.h
+++ b/Engine/Plugins/VirtualProduction/DMX/DMXPixelMapping/Source/DMXPixelMappingRuntime/Public/Components/DMXPixelMappingRendererComponent.h
@@ -2,10 +2,10 @@
 
 #pragma once
 
-#include "Components/DMXPixelMappingOutputComponent.h"
-
 #include "IDMXPixelMappingRenderer.h"
+#include "Components/DMXPixelMappingOutputComponent.h"
 #include "Library/DMXEntityReference.h"
+#include "Rendering/IDMXPixelMappingRenderInputTextureProxy.h"
 
 #include "Templates/SubclassOf.h"
 
@@ -35,12 +35,10 @@ public:
 	/** Default Constructor */
 	UDMXPixelMappingRendererComponent();
 
-	/** Destructor */
-	~UDMXPixelMappingRendererComponent();
-
 	//~ Begin UObject implementation
 	virtual void Serialize(FArchive& Ar) override;
 	virtual void PostInitProperties() override;
+	virtual void PostLoad() override;
 #if WITH_EDITOR
 	virtual void PostEditChangeChainProperty(FPropertyChangedChainEvent& PropertyChangedChainEvent) override;
 #endif // WITH_EDITOR
@@ -65,9 +63,6 @@ public:
 	/** Get reference to the active input texture */
 	UTexture* GetRendererInputTexture() const;
 
-	/** Get renderer interfece */
-	const TSharedPtr<IDMXPixelMappingRenderer>& GetRenderer() { return PixelMappingRenderer; }
-
 	/**
 	 * Get pixel position in downsample buffer target based on pixel index
 	 *
@@ -79,6 +74,9 @@ public:
 	/** Get active world. It could be editor or build world */
 	UWorld* GetWorld() const;
 
+	/** Get renderer interface */
+	const TSharedPtr<IDMXPixelMappingRenderer>& GetRenderer() { return PixelMappingRenderer; }
+
 #if WITH_EDITOR
 	/**
 	 * Take of container widget which is holds widget for all child components.
@@ -95,7 +93,7 @@ public:
 	UFUNCTION(BlueprintCallable, Category = "DMX|PixelMapping")
 	void RendererInputTexture();
 
-	/** Create or update size of  buffer target for rendering downsample pixels */
+	/** Create or update size of buffer target for rendering downsample pixels */
 	void CreateOrUpdateDownsampleBufferTarget();
 
 	/** 
@@ -132,20 +130,12 @@ public:
 	void EmptyDownsampleBuffer();
 
 private:
-	/** Resize input target based on X and Y input material size  */
-	void ResizeMaterialRenderTarget(int32 InSizeX, int32 InSizeY);
-
 	/** Generate new input widget based on UMG */
 	void UpdateInputWidget(TSubclassOf<UUserWidget> InInputWidget);
 
 	/** Resize output texture for editor preview */
 	void ResizePreviewRenderTarget(uint32 InSizeX, uint32 InSizeY);
 
-#if WITH_EDITOR
-	/** Map changer handler */
-	void OnMapChanged(UWorld* InWorld, EMapChangeType MapChangeType);
-#endif
-
 	/** Initialize all textures and creation or loading asset */
 	void Initialize();
 
@@ -175,10 +165,14 @@ public:
 	/** UMG to Downsampling */
 	UPROPERTY(EditAnywhere, BlueprintReadOnly, Category = "Render Settings")
 	TSubclassOf<UUserWidget> InputWidget;
-	
+
 	/** The brightness of the renderer */
 	UPROPERTY(EditAnywhere, BlueprintReadOnly, Category = "Render Settings", meta = (ClampMin = "0", ClampMax = "1", UIMin = "0", UIMax = "1"))
-	float Brightness;
+	float Brightness = 1.f;
+
+	/** Post process material applied to the rendered input */
+	UPROPERTY(EditAnywhere, BlueprintReadOnly, Category = "Post Process")
+	TObjectPtr<UMaterialInterface> PostProcessMaterial;
 
 	/** Layout script for the children of this component (hidden in customizations and displayed in its own panel). */
 	UPROPERTY(EditAnywhere, Instanced, Category = "Layout")
@@ -193,15 +187,21 @@ public:
 #endif // WITH_EDITOR
 
 private:
-#if WITH_EDITORONLY_DATA
-	/** Editor preview output target */
-	UPROPERTY(Transient)
-	TObjectPtr<UTextureRenderTarget2D> PreviewRenderTarget;
-#endif // WITH_EDITORONLY_DATA
+	/** Retrieve total count of all output targets that support shared rendering and updates a counter. O(n) */
+	int32 GetTotalDownsamplePixelCount();
 
-	/** Material of UMG texture to downsample */
-	UPROPERTY(Transient)
-	TObjectPtr<UTextureRenderTarget2D> InputRenderTarget;
+	/** Helper function checks the downsample pixel range */
+	bool IsPixelRangeValid(const int32 InDownsamplePixelIndexStart, const int32 InDownsamplePixelIndexEnd) const;
+
+	/** Updates the Input Texture Render Proxy, depending on the selected source */
+	void UpdateInputTextureRenderProxy();
+
+	/** Proxy responsible to render the input texture */
+	TSharedPtr<UE::DMXPixelMapping::Rendering::Private::IDMXPixelMappingRenderInputTextureProxy> RenderInputTextureProxy;
+
+	/** Parameters by which the input texture is rendered */
+	UPROPERTY(EditAnywhere, Category = "Post Process", Meta = (ShowOnlyInnerProperties, AllowPrivateAccess = true))
+	FDMXPixelMappingRenderInputTextureParameters RenderInputTextureParams;
 
 	/** Reference to renderer */
 	TSharedPtr<IDMXPixelMappingRenderer> PixelMappingRenderer;
@@ -213,26 +213,23 @@ private:
 #if WITH_EDITORONLY_DATA
 	/** Canvas for all UI downsamping component widgets */
 	TSharedPtr<SConstraintCanvas> ComponentsCanvas;
-
-	/** Change level Delegate */
-	FDelegateHandle OnChangeLevelHandle;
 #endif // WITH_EDITORONLY_DATA
 
-	/** Retrieve total count of all output targets that support shared rendering and updates a counter. O(n) */
-	int32 GetTotalDownsamplePixelCount();
-
-	/** Helper function checks the downsample pixel range */
-	bool IsPixelRangeValid(const int32 InDownsamplePixelIndexStart, const int32 InDownsamplePixelIndexEnd) const;
-
 	/** GPU downsample pixel buffer target texture */
-	UPROPERTY(Transient)
+	UPROPERTY(Transient, DuplicateTransient)
 	TObjectPtr<UTextureRenderTarget2D> DownsampleBufferTarget;
 
+#if WITH_EDITORONLY_DATA
+	/** Editor preview output target */
+	UPROPERTY(Transient, DuplicateTransient)
+	TObjectPtr<UTextureRenderTarget2D> PreviewRenderTarget;
+#endif // WITH_EDITORONLY_DATA
+
 	/** CPU downsample pixel buffer */
 	TArray<FLinearColor> DownsampleBuffer;
 
 	/** Counter for all pixels from child components */
-	int32 DownsamplePixelCount;
+	int32 DownsamplePixelCount = 0;
 
 	/** Critical section for set, update and get color array */
 	mutable FCriticalSection DownsampleBufferCS;
diff --git a/Engine/Plugins/VirtualProduction/DMX/DMXPixelMapping/Source/DMXPixelMappingRuntime/Public/Rendering/IDMXPixelMappingRenderInputTextureProxy.h b/Engine/Plugins/VirtualProduction/DMX/DMXPixelMapping/Source/DMXPixelMappingRuntime/Public/Rendering/IDMXPixelMappingRenderInputTextureProxy.h
new file mode 100644
--- /dev/null
+++ b/Engine/Plugins/VirtualProduction/DMX/DMXPixelMapping/Source/DMXPixelMappingRuntime/Public/Rendering/IDMXPixelMappingRenderInputTextureProxy.h
@@ -0,0 +1,63 @@
+// Copyright Epic Games, Inc. All Rights Reserved.
+
+#pragma once
+
+#include "IDMXPixelMappingRenderInputTextureProxy.generated.h"
+
+class UMaterialInstanceDynamic;
+class UTexture;
+
+
+/** Parameters to render the input texture */
+USTRUCT(BlueprintType)
+struct FDMXPixelMappingRenderInputTextureParameters
+{
+	GENERATED_BODY()
+
+	/** Post process material applied to the rendered input */
+	UPROPERTY(Transient)
+	TObjectPtr<UMaterialInstanceDynamic> PostProcessMID;
+
+	/** If true, the post process material is applied each downsample pass, otherwise only once after the last pass */
+	UPROPERTY(EditAnywhere, Category = "Post Process")
+	bool bApplyPostProcessMaterialEachDownsamplePass = true;
+
+	/** Number of times the pixelmapping input is downsampled */
+	UPROPERTY(EditAnywhere, Category = "Post Process", Meta = (ClampMin = "0", ClampMax = "256", UIMin = "0", UIMax = "16"))
+	int32 NumDownSamplePasses = 0;
+
+	/** Blur distance applied, only applicable if the post process matierial has a "BlurDistance" parameter */
+	UPROPERTY(EditAnywhere, Category = "Post Process", Meta = (UIMin = "0", UIMax = "1", SliderExponent = 2.0))
+	float BlurDistance = 0.02f;
+
+	/** Parameter name of the post process material to receive the input texture */
+	UPROPERTY(EditAnywhere, AdvancedDisplay, Category = "Post Process")
+	FName InputTextureParameterName = "InputTexture";
+
+	/** The blur distance parameter name */
+	FName BlurDistanceParamtereName = "BlurDistance";
+
+	/** Size of the rendered texture */
+	FVector2D OutputSize{ 1.f, 1.f };
+};
+
+namespace UE::DMXPixelMapping::Rendering::Private
+{
+	/** Renders the input to be used by pixel mapping. */
+	class IDMXPixelMappingRenderInputTextureProxy
+		: public TSharedFromThis<IDMXPixelMappingRenderInputTextureProxy>
+	{
+	public:
+		virtual ~IDMXPixelMappingRenderInputTextureProxy() {}
+
+		/** Updates parameters */
+		virtual void UpdateParameters(const FDMXPixelMappingRenderInputTextureParameters& InParameters) = 0;
+
+		/** Renders the input texture, optionally using a post process material */
+		virtual void Render() = 0;
+
+		/** Returns the currently rendered texture, or nullptr if no texture was rendered yet */
+		virtual UTexture* GetRenderedTexture() const = 0;
+	};
+
+}
