DMX - Fix CIS NoPCH Compile Win64

diff --git a/Engine/Plugins/VirtualProduction/DMX/DMXPixelMapping/Source/DMXPixelMappingRuntime/Private/Components/DMXPixelMappingRendererComponent.cpp b/Engine/Plugins/VirtualProduction/DMX/DMXPixelMapping/Source/DMXPixelMappingRuntime/Private/Components/DMXPixelMappingRendererComponent.cpp
--- a/Engine/Plugins/VirtualProduction/DMX/DMXPixelMapping/Source/DMXPixelMappingRuntime/Private/Components/DMXPixelMappingRendererComponent.cpp
+++ b/Engine/Plugins/VirtualProduction/DMX/DMXPixelMapping/Source/DMXPixelMappingRuntime/Private/Components/DMXPixelMappingRendererComponent.cpp
@@ -21,6 +21,7 @@
 #include "Modulators/DMXModulator.h"
 #include "TextureResource.h"
 #include "UObject/ConstructorHelpers.h"
+#include "UObject/Package.h"
 #include "Widgets/Layout/SConstraintCanvas.h"
 
 
