DMX - PixelMapping Fixture and Matrix Fixtures now order their properties identically

diff --git a/Engine/Plugins/VirtualProduction/DMX/DMXPixelMapping/Source/DMXPixelMappingEditor/Private/Customizations/DMXPixelMappingDetailCustomization_Matrix.cpp b/Engine/Plugins/VirtualProduction/DMX/DMXPixelMapping/Source/DMXPixelMappingEditor/Private/Customizations/DMXPixelMappingDetailCustomization_Matrix.cpp
--- a/Engine/Plugins/VirtualProduction/DMX/DMXPixelMapping/Source/DMXPixelMappingEditor/Private/Customizations/DMXPixelMappingDetailCustomization_Matrix.cpp
+++ b/Engine/Plugins/VirtualProduction/DMX/DMXPixelMapping/Source/DMXPixelMappingEditor/Private/Customizations/DMXPixelMappingDetailCustomization_Matrix.cpp
@@ -40,6 +40,48 @@ void FDMXPixelMappingDetailCustomization_Matrix::CustomizeDetails(IDetailLayoutB
 	InDetailLayout.HideProperty(GET_MEMBER_NAME_CHECKED(UDMXPixelMappingMatrixComponent, LayoutScript));
 
 	CreateModulatorDetails(InDetailLayout);
+
+	// Sort categories
+	InDetailLayout.SortCategories([](const TMap<FName, IDetailCategoryBuilder*>& CategoryMap)
+		{
+			int32 MinSortOrder = TNumericLimits<int32>::Max();
+			for (const TPair<FName, IDetailCategoryBuilder*>& Pair : CategoryMap)
+			{
+				int32 SortOrder = Pair.Value->GetSortOrder();
+				MinSortOrder = FMath::Min(SortOrder, MinSortOrder);
+			}
+
+			IDetailCategoryBuilder* const* ColorSpaceCategoryPtr = CategoryMap.Find("Color Space");
+			if (ColorSpaceCategoryPtr)
+			{
+				(*ColorSpaceCategoryPtr)->SetSortOrder(MinSortOrder - 3);
+			}
+
+			// Either 'RGB' 'XY' or 'XYZ' is displayed
+			IDetailCategoryBuilder* const* RGBCategoryPtr = CategoryMap.Find("RGB");
+			if (RGBCategoryPtr)
+			{
+				(*RGBCategoryPtr)->SetSortOrder(MinSortOrder - 2);
+			}
+
+			IDetailCategoryBuilder* const* XYCategoryPtr = CategoryMap.Find("XY");
+			if (XYCategoryPtr)
+			{
+				(*XYCategoryPtr)->SetSortOrder(MinSortOrder - 2);
+			}
+
+			IDetailCategoryBuilder* const* XYZCategoryPtr = CategoryMap.Find("XYZ");
+			if (XYZCategoryPtr)
+			{
+				(*XYZCategoryPtr)->SetSortOrder(MinSortOrder - 2);
+			}
+
+			IDetailCategoryBuilder* const* IntensityCategoryPtr = CategoryMap.Find("Luminance");
+			if (IntensityCategoryPtr)
+			{
+				(*IntensityCategoryPtr)->SetSortOrder(MinSortOrder - 1);
+			}
+		});
 }
 
 void FDMXPixelMappingDetailCustomization_Matrix::CreateModulatorDetails(IDetailLayoutBuilder& InDetailLayout)
diff --git a/Engine/Plugins/VirtualProduction/DMX/DMXPixelMapping/Source/DMXPixelMappingRuntime/Public/Components/DMXPixelMappingMatrixComponent.h b/Engine/Plugins/VirtualProduction/DMX/DMXPixelMapping/Source/DMXPixelMappingRuntime/Public/Components/DMXPixelMappingMatrixComponent.h
--- a/Engine/Plugins/VirtualProduction/DMX/DMXPixelMapping/Source/DMXPixelMappingRuntime/Public/Components/DMXPixelMappingMatrixComponent.h
+++ b/Engine/Plugins/VirtualProduction/DMX/DMXPixelMapping/Source/DMXPixelMappingRuntime/Public/Components/DMXPixelMappingMatrixComponent.h
@@ -105,7 +105,7 @@ public:
 	FDMXEntityFixturePatchRef FixturePatchMatrixRef_DEPRECATED;
 #endif // WITH_EDITORONLY_DATA
 
-	UPROPERTY(VisibleAnywhere, BlueprintReadOnly, Category = "Matrix Settings")
+	UPROPERTY(VisibleAnywhere, BlueprintReadOnly, Category = "Selected Patch")
 	FDMXEntityFixturePatchRef FixturePatchRef;
 
 	/** Sets which color space Pixel Mapping sends */
