DMX - Refactor pixelmapping preprocess and texture filtering, fix various issues

diff --git a/Engine/Plugins/VirtualProduction/DMX/DMXPixelMapping/Source/DMXPixelMappingEditor/Private/Customizations/DMXPixelMappingDetailCustomization_Renderer.cpp b/Engine/Plugins/VirtualProduction/DMX/DMXPixelMapping/Source/DMXPixelMappingEditor/Private/Customizations/DMXPixelMappingDetailCustomization_Renderer.cpp
--- a/Engine/Plugins/VirtualProduction/DMX/DMXPixelMapping/Source/DMXPixelMappingEditor/Private/Customizations/DMXPixelMappingDetailCustomization_Renderer.cpp
+++ b/Engine/Plugins/VirtualProduction/DMX/DMXPixelMapping/Source/DMXPixelMappingEditor/Private/Customizations/DMXPixelMappingDetailCustomization_Renderer.cpp
@@ -10,14 +10,13 @@
 #include "Engine/Texture.h"
 #include "Materials/Material.h"
 #include "TextureResource.h"
+#include "Widgets/Text/STextBlock.h"
 
 #define LOCTEXT_NAMESPACE "DMXPixelMappingDetailCustomization_Renderer"
 
 
 void FDMXPixelMappingDetailCustomization_Renderer::CustomizeDetails(IDetailLayoutBuilder& DetailLayout)
 {
-	IDetailCategoryBuilder& RenderSettingsSettingsCategory = DetailLayout.EditCategory("Render Settings", FText::GetEmpty(), ECategoryPriority::Important);
-
 	// Register all handles
 	RendererTypePropertyHandle = DetailLayout.GetProperty(GET_MEMBER_NAME_CHECKED(UDMXPixelMappingRendererComponent, RendererType));
 	InputTexturePropertyHandle = DetailLayout.GetProperty(GET_MEMBER_NAME_CHECKED(UDMXPixelMappingRendererComponent, InputTexture));
@@ -35,19 +34,30 @@ void FDMXPixelMappingDetailCustomization_Renderer::CustomizeDetails(IDetailLayou
 	RendererTypePropertyHandle->SetOnPropertyValueChanged(OnInputTexturePropertyChangedDelegate);
 
 	// Hide Output Component properties
-	TSharedRef<IPropertyHandle> PositionXPropertyHandle = DetailLayout.GetProperty(UDMXPixelMappingOutputComponent::GetPositionXPropertyName(), UDMXPixelMappingOutputComponent::StaticClass());
-	TSharedRef<IPropertyHandle> PositionYPropertyHandle = DetailLayout.GetProperty(UDMXPixelMappingOutputComponent::GetPositionYPropertyName(), UDMXPixelMappingOutputComponent::StaticClass());
-	TSharedRef<IPropertyHandle> IsLockInDesignerPropertyHandle = DetailLayout.GetProperty(UDMXPixelMappingOutputComponent::GetLockInDesignerPropertyName(), UDMXPixelMappingOutputComponent::StaticClass());
-	TSharedRef<IPropertyHandle> IsVisibleInDesignerPropertyHandle = DetailLayout.GetProperty(UDMXPixelMappingOutputComponent::GetVisibleInDesignerPropertyName(), UDMXPixelMappingOutputComponent::StaticClass());
-	TSharedRef<IPropertyHandle> PixelBlendingPropertyHandle = DetailLayout.GetProperty(GET_MEMBER_NAME_CHECKED(UDMXPixelMappingOutputComponent, CellBlendingQuality), UDMXPixelMappingOutputComponent::StaticClass());
+	const TSharedRef<IPropertyHandle> PositionXPropertyHandle = DetailLayout.GetProperty(UDMXPixelMappingOutputComponent::GetPositionXPropertyName(), UDMXPixelMappingOutputComponent::StaticClass());
 	DetailLayout.HideProperty(PositionXPropertyHandle);
+
+	const TSharedRef<IPropertyHandle> PositionYPropertyHandle = DetailLayout.GetProperty(UDMXPixelMappingOutputComponent::GetPositionYPropertyName(), UDMXPixelMappingOutputComponent::StaticClass());
 	DetailLayout.HideProperty(PositionYPropertyHandle);
+
+	const TSharedRef<IPropertyHandle> SizeXPropertyHandle = DetailLayout.GetProperty(UDMXPixelMappingOutputComponent::GetSizeXPropertyName(), UDMXPixelMappingOutputComponent::StaticClass());
+	DetailLayout.HideProperty(SizeXPropertyHandle);
+
+	const TSharedRef<IPropertyHandle> SizeYPropertyHandle = DetailLayout.GetProperty(UDMXPixelMappingOutputComponent::GetSizeYPropertyName(), UDMXPixelMappingOutputComponent::StaticClass());
+	DetailLayout.HideProperty(SizeYPropertyHandle);
+
+	const TSharedRef<IPropertyHandle> IsLockInDesignerPropertyHandle = DetailLayout.GetProperty(UDMXPixelMappingOutputComponent::GetLockInDesignerPropertyName(), UDMXPixelMappingOutputComponent::StaticClass());
 	DetailLayout.HideProperty(IsLockInDesignerPropertyHandle);
+
+	const TSharedRef<IPropertyHandle> IsVisibleInDesignerPropertyHandle = DetailLayout.GetProperty(UDMXPixelMappingOutputComponent::GetVisibleInDesignerPropertyName(), UDMXPixelMappingOutputComponent::StaticClass());
 	DetailLayout.HideProperty(IsVisibleInDesignerPropertyHandle);
+
+	const TSharedRef<IPropertyHandle> PixelBlendingPropertyHandle = DetailLayout.GetProperty(GET_MEMBER_NAME_CHECKED(UDMXPixelMappingOutputComponent, CellBlendingQuality), UDMXPixelMappingOutputComponent::StaticClass());
 	DetailLayout.HideProperty(PixelBlendingPropertyHandle);
 
 	// Add properties
-	RenderSettingsSettingsCategory.AddProperty(RendererTypePropertyHandle);
+	IDetailCategoryBuilder& RenderSettingsCategory = DetailLayout.EditCategory("Render Settings", FText::GetEmpty(), ECategoryPriority::Important);
+	RenderSettingsCategory.AddProperty(RendererTypePropertyHandle);
 
 	// Get editing UObject
 	TArray<TWeakObjectPtr<UObject>> OuterObjects;
@@ -61,40 +71,30 @@ void FDMXPixelMappingDetailCustomization_Renderer::CustomizeDetails(IDetailLayou
 			// Remember the selected renderer type
 			PreviousRendererType = RendererComponent->RendererType;
 
-		// Add Warning InputTexture Row
-		AddInputTextureWarning(RenderSettingsSettingsCategory);
-
-		// Add non UI Material Warning
-		AddMaterialWarning(RenderSettingsSettingsCategory);
-
-		// Retister properties
-		RenderSettingsSettingsCategory.AddProperty(InputTexturePropertyHandle)
-			.Visibility(
-				TAttribute<EVisibility>::Create(
-						TAttribute<EVisibility>::FGetter::CreateLambda([this]
-							{
-								return IsSelectedRendererType(EDMXPixelMappingRendererType::Texture) ? EVisibility::Visible : EVisibility::Hidden;
-							}
-					)));
-
-		RenderSettingsSettingsCategory.AddProperty(InputMaterialPropertyHandle)
-			.Visibility(
-				TAttribute<EVisibility>::Create(
-						TAttribute<EVisibility>::FGetter::CreateLambda([this]
-							{
-								return IsSelectedRendererType(EDMXPixelMappingRendererType::Material) ? EVisibility::Visible : EVisibility::Hidden;
-							}
-					)));
-
-		RenderSettingsSettingsCategory.AddProperty(InputWidgetPropertyHandle)
-			.Visibility(
-				TAttribute<EVisibility>::Create(
-						TAttribute<EVisibility>::FGetter::CreateLambda([this]
-							{
-								return IsSelectedRendererType(EDMXPixelMappingRendererType::UMG) ? EVisibility::Visible : EVisibility::Hidden;
-							}
-					)));
-
+			// Add Warning InputTexture Row
+			AddInputTextureWarning(RenderSettingsCategory);
+
+			// Add non UI Material Warning
+			AddMaterialWarning(RenderSettingsCategory);
+
+			// Register properties
+			RenderSettingsCategory.AddProperty(InputTexturePropertyHandle)
+				.Visibility(TAttribute<EVisibility>::Create(TAttribute<EVisibility>::FGetter::CreateLambda([this]
+					{
+						return IsSelectedRendererType(EDMXPixelMappingRendererType::Texture) ? EVisibility::Visible : EVisibility::Hidden;
+					})));
+
+			RenderSettingsCategory.AddProperty(InputMaterialPropertyHandle)
+				.Visibility(TAttribute<EVisibility>::Create(TAttribute<EVisibility>::FGetter::CreateLambda([this]
+					{
+						return IsSelectedRendererType(EDMXPixelMappingRendererType::Material) ? EVisibility::Visible : EVisibility::Hidden;
+					})));
+
+			RenderSettingsCategory.AddProperty(InputWidgetPropertyHandle)
+				.Visibility(TAttribute<EVisibility>::Create(TAttribute<EVisibility>::FGetter::CreateLambda([this]
+					{
+						return IsSelectedRendererType(EDMXPixelMappingRendererType::UMG) ? EVisibility::Visible : EVisibility::Hidden;
+					})));
 		}
 	}
 }
@@ -155,25 +155,25 @@ void FDMXPixelMappingDetailCustomization_Renderer::AddInputTextureWarning(IDetai
 	const FSlateBrush* WarningIcon = FAppStyle::GetBrush("SettingsEditor.WarningIcon");
 
 	InCategory.AddCustomRow(FText::GetEmpty())
-			.Visibility(TAttribute<EVisibility>(this, &FDMXPixelMappingDetailCustomization_Renderer::GetInputTextureWarning))
-			.WholeRowContent()
-			[
-				SNew(SWarningOrErrorBox)
-				.MessageStyle(EMessageStyle::Warning)
-				.Message(this, &FDMXPixelMappingDetailCustomization_Renderer::GetInputTextureWarningText)
-			];
+		.Visibility(TAttribute<EVisibility>(this, &FDMXPixelMappingDetailCustomization_Renderer::GetInputTextureWarning))
+		.WholeRowContent()
+		[
+			SNew(SWarningOrErrorBox)
+			.MessageStyle(EMessageStyle::Warning)
+			.Message(this, &FDMXPixelMappingDetailCustomization_Renderer::GetInputTextureWarningText)
+		];
 }
 
 void FDMXPixelMappingDetailCustomization_Renderer::AddMaterialWarning(IDetailCategoryBuilder& InCategory)
 {
 	InCategory.AddCustomRow(FText::GetEmpty())
-			.Visibility(TAttribute<EVisibility>(this, &FDMXPixelMappingDetailCustomization_Renderer::GetMaterialWarningVisibility))
-			.WholeRowContent()
-			[
-				SNew(SWarningOrErrorBox)
-				.MessageStyle(EMessageStyle::Warning)
-				.Message(LOCTEXT("WarningNonUIMaterial", "Selected Material is not UI Material.\nChange Material Domain to User Interface.\nOr select another Material."))
-			];
+		.Visibility(TAttribute<EVisibility>(this, &FDMXPixelMappingDetailCustomization_Renderer::GetMaterialWarningVisibility))
+		.WholeRowContent()
+		[
+			SNew(SWarningOrErrorBox)
+			.MessageStyle(EMessageStyle::Warning)
+			.Message(LOCTEXT("WarningNonUIMaterial", "Selected Material is not UI Material.\nChange Material Domain to User Interface.\nOr select another Material."))
+		];
 }
 
 EVisibility FDMXPixelMappingDetailCustomization_Renderer::GetMaterialWarningVisibility() const
diff --git a/Engine/Plugins/VirtualProduction/DMX/DMXPixelMapping/Source/DMXPixelMappingEditor/Private/Customizations/DMXPixelMappingPreprocessRendererDetails.cpp b/Engine/Plugins/VirtualProduction/DMX/DMXPixelMapping/Source/DMXPixelMappingEditor/Private/Customizations/DMXPixelMappingPreprocessRendererDetails.cpp
new file mode 100644
--- /dev/null
+++ b/Engine/Plugins/VirtualProduction/DMX/DMXPixelMapping/Source/DMXPixelMappingEditor/Private/Customizations/DMXPixelMappingPreprocessRendererDetails.cpp
@@ -0,0 +1,78 @@
+// Copyright Epic Games, Inc. All Rights Reserved.
+
+#include "DMXPixelMappingPreprocessRendererDetails.h"
+
+#include "Components/DMXPixelMappingRendererComponent.h"
+#include "DMXPixelMappingPreprocessRenderer.h"
+#include "DetailLayoutBuilder.h"
+#include "IPropertyUtilities.h"
+#include "MaterialTypes.h"
+#include "Materials/MaterialInterface.h"
+#include "PropertyHandle.h"
+
+
+namespace UE::DMXPixelMapping::Customizations
+{
+	TSharedRef<IDetailCustomization> FDMXPixelMappingPreprocessRendererDetails::MakeInstance()
+	{
+		return MakeShared<FDMXPixelMappingPreprocessRendererDetails>();
+	}
+
+	void FDMXPixelMappingPreprocessRendererDetails::CustomizeDetails(IDetailLayoutBuilder& DetailLayout)
+	{
+		PropertyUtilities = DetailLayout.GetPropertyUtilities();
+
+		// Handle blur distance property enabled
+		FilterMaterialHandle = DetailLayout.GetProperty(UDMXPixelMappingPreprocessRenderer::GetFilterMaterialMemberNameChecked());
+		FilterMaterialHandle->SetOnPropertyValueChanged(FSimpleDelegate::CreateSP(this, &FDMXPixelMappingPreprocessRendererDetails::OnFilterMaterialChanged));
+		FilterMaterialHandle->SetOnChildPropertyValueChanged(FSimpleDelegate::CreateSP(this, &FDMXPixelMappingPreprocessRendererDetails::OnFilterMaterialChanged));
+
+		if (!IsBlurDistanceMaterialParameterNameExisting())
+		{
+			const TSharedRef<IPropertyHandle> BlurDistanceHandle = DetailLayout.GetProperty(UDMXPixelMappingPreprocessRenderer::GetBlurDistanceMemberNameChecked());
+			BlurDistanceHandle->MarkHiddenByCustomization();
+		}
+	}
+
+	bool FDMXPixelMappingPreprocessRendererDetails::IsBlurDistanceMaterialParameterNameExisting() const
+	{
+		const TArray<TWeakObjectPtr<UObject>> EditedObjects = PropertyUtilities->GetSelectedObjects();
+
+		for (TWeakObjectPtr<UObject> EditedObject : EditedObjects)
+		{
+			UDMXPixelMappingRendererComponent* RendererComponent = Cast<UDMXPixelMappingRendererComponent>(EditedObject);
+			if (!RendererComponent)
+			{
+				continue;
+			}
+
+			UDMXPixelMappingPreprocessRenderer* PreprocessRenderer = RendererComponent->GetPreprocessRenderer();
+			if (!ensureMsgf(PreprocessRenderer, TEXT("Missing default subobbject RenderInputTextureProxy in %s"), *RendererComponent->GetName()))
+			{
+				continue;
+			}
+
+			UObject* FilterMaterialObject = nullptr;
+			if (FilterMaterialHandle->GetValue(FilterMaterialObject) != FPropertyAccess::Success)
+			{
+				continue;
+			}
+
+			UMaterialInterface* FilterMaterial = Cast<UMaterialInterface>(FilterMaterialObject);
+			if (!FilterMaterial)
+			{
+				continue;
+			}
+
+			FMaterialParameterMetadata ParameterMetadata;
+			return FilterMaterial->GetParameterValue(EMaterialParameterType::Scalar, FMemoryImageMaterialParameterInfo(PreprocessRenderer->GetBlurDistanceParameterName()), ParameterMetadata);
+		}
+
+		return false;
+	}
+
+	void FDMXPixelMappingPreprocessRendererDetails::OnFilterMaterialChanged()
+	{
+		PropertyUtilities->ForceRefresh();
+	}
+}
diff --git a/Engine/Plugins/VirtualProduction/DMX/DMXPixelMapping/Source/DMXPixelMappingEditor/Private/Customizations/DMXPixelMappingPreprocessRendererDetails.h b/Engine/Plugins/VirtualProduction/DMX/DMXPixelMapping/Source/DMXPixelMappingEditor/Private/Customizations/DMXPixelMappingPreprocessRendererDetails.h
new file mode 100644
--- /dev/null
+++ b/Engine/Plugins/VirtualProduction/DMX/DMXPixelMapping/Source/DMXPixelMappingEditor/Private/Customizations/DMXPixelMappingPreprocessRendererDetails.h
@@ -0,0 +1,35 @@
+// Copyright Epic Games, Inc. All Rights Reserved.
+
+#pragma once
+
+#include "IDetailCustomization.h"
+
+class IPropertyHandle;
+class IPropertyUtilities;
+
+
+namespace UE::DMXPixelMapping::Customizations
+{
+	class FDMXPixelMappingPreprocessRendererDetails
+		: public IDetailCustomization
+	{
+	public:
+		static TSharedRef<IDetailCustomization> MakeInstance();
+
+		/** IDetailCustomization interface */
+		virtual void CustomizeDetails(IDetailLayoutBuilder& DetailLayout) override;
+
+	private:
+		/** Called when the filter material instance dynamic property changed */
+		void OnFilterMaterialChanged();
+
+		/** Tests if the blur distance parameter name exists on the filter material */
+		bool IsBlurDistanceMaterialParameterNameExisting() const;
+
+		/** Property handle for the FilterMaterial property */
+		TSharedPtr<IPropertyHandle> FilterMaterialHandle;
+
+		/** Property utilities for this customization */
+		TSharedPtr<IPropertyUtilities> PropertyUtilities;
+	};
+}
diff --git a/Engine/Plugins/VirtualProduction/DMX/DMXPixelMapping/Source/DMXPixelMappingEditor/Private/Views/SDMXPixelMappingInputSourceView.cpp b/Engine/Plugins/VirtualProduction/DMX/DMXPixelMapping/Source/DMXPixelMappingEditor/Private/Views/SDMXPixelMappingInputSourceView.cpp
--- a/Engine/Plugins/VirtualProduction/DMX/DMXPixelMapping/Source/DMXPixelMappingEditor/Private/Views/SDMXPixelMappingInputSourceView.cpp
+++ b/Engine/Plugins/VirtualProduction/DMX/DMXPixelMapping/Source/DMXPixelMappingEditor/Private/Views/SDMXPixelMappingInputSourceView.cpp
@@ -2,17 +2,18 @@
 
 #include "Views/SDMXPixelMappingInputSourceView.h"
 
-#include "DMXPixelMapping.h"
-#include "DMXPixelMappingComponentReference.h"
 #include "Components/DMXPixelMappingRendererComponent.h"
 #include "Components/DMXPixelMappingRootComponent.h"
 #include "Customizations/DMXPixelMappingDetailCustomization_Renderer.h"
-#include "Toolkits/DMXPixelMappingToolkit.h"
-
+#include "Customizations/DMXPixelMappingPreprocessRendererDetails.h"
+#include "DMXPixelMapping.h"
+#include "DMXPixelMappingComponentReference.h"
+#include "DMXPixelMappingPreprocessRenderer.h"
 #include "Editor.h"
 #include "IDetailsView.h"
 #include "PropertyEditorDelegates.h"
 #include "PropertyEditorModule.h"
+#include "Toolkits/DMXPixelMappingToolkit.h"
 #include "Widgets/SBoxPanel.h"
 #include "Widgets/Images/SImage.h"
 #include "Widgets/Input/SButton.h"
@@ -178,6 +179,11 @@ void SDMXPixelMappingInputSourceView::ForceRefresh()
 
 			FOnGetDetailCustomizationInstance RendererCustomizationInstance = FOnGetDetailCustomizationInstance::CreateStatic(&FDMXPixelMappingDetailCustomization_Renderer::MakeInstance, WeakToolkit);
 			DetailsView->RegisterInstancedCustomPropertyLayout(UDMXPixelMappingRendererComponent::StaticClass(), RendererCustomizationInstance);
+
+			using namespace UE::DMXPixelMapping::Customizations;
+			FOnGetDetailCustomizationInstance RenderInputTextureProxyCustomizationInstance = FOnGetDetailCustomizationInstance::CreateStatic(&FDMXPixelMappingPreprocessRendererDetails::MakeInstance);
+			DetailsView->RegisterInstancedCustomPropertyLayout(UDMXPixelMappingPreprocessRenderer::StaticClass(), RenderInputTextureProxyCustomizationInstance);
+			
 			DetailsView->SetObject(RendererComponent);
 
 			BodyContent->AddSlot()
diff --git a/Engine/Plugins/VirtualProduction/DMX/DMXPixelMapping/Source/DMXPixelMappingRenderer/Private/DMXPixelMappingPostProcessProxy.cpp b/Engine/Plugins/VirtualProduction/DMX/DMXPixelMapping/Source/DMXPixelMappingRenderer/Private/DMXPixelMappingPostProcessProxy.cpp
deleted file mode 100644
diff --git a/Engine/Plugins/VirtualProduction/DMX/DMXPixelMapping/Source/DMXPixelMappingRenderer/Private/DMXPixelMappingPostProcessProxy.h b/Engine/Plugins/VirtualProduction/DMX/DMXPixelMapping/Source/DMXPixelMappingRenderer/Private/DMXPixelMappingPostProcessProxy.h
deleted file mode 100644
diff --git a/Engine/Plugins/VirtualProduction/DMX/DMXPixelMapping/Source/DMXPixelMappingRenderer/Private/DMXPixelMappingPreprocessRenderer.cpp b/Engine/Plugins/VirtualProduction/DMX/DMXPixelMapping/Source/DMXPixelMappingRenderer/Private/DMXPixelMappingPreprocessRenderer.cpp
new file mode 100644
--- /dev/null
+++ b/Engine/Plugins/VirtualProduction/DMX/DMXPixelMapping/Source/DMXPixelMappingRenderer/Private/DMXPixelMappingPreprocessRenderer.cpp
@@ -0,0 +1,125 @@
+// Copyright Epic Games, Inc. All Rights Reserved.
+
+#include "DMXPixelMappingPreprocessRenderer.h"
+
+#include "Engine/TextureRenderTarget2D.h"
+#include "IDMXPixelMappingRenderer.h"
+#include "Materials/MaterialInstanceDynamic.h"
+#include "PreprocessApplyFilterMaterialProxy.h"
+#include "PreprocessRenderInputMaterialProxy.h"
+#include "PreprocessRenderInputTextureProxy.h"
+#include "PreprocessRenderInputUserWidgetProxy.h"
+#include "UObject/Package.h"
+
+
+void UDMXPixelMappingPreprocessRenderer::SetInputTexture(UTexture* InTexture)
+{
+	using namespace UE::DMXPixelMapping::Rendering::Preprocess::Private;
+	RenderInputProxy = MakeShared<FPreprocessRenderInputTextureProxy>(InTexture);
+	ApplyMaterialProxy = MakeShared<FPreprocessApplyFilterMaterialProxy>();
+}
+
+void UDMXPixelMappingPreprocessRenderer::SetInputMaterial(UMaterialInterface* InMaterial)
+{
+	using namespace UE::DMXPixelMapping::Rendering::Preprocess::Private;
+	RenderInputProxy = MakeShared<FPreprocessRenderInputMaterialProxy>(InMaterial, InputSize);
+	ApplyMaterialProxy = MakeShared<FPreprocessApplyFilterMaterialProxy>();
+}
+
+void UDMXPixelMappingPreprocessRenderer::SetInputUserWidget(UUserWidget* InUserWidget)
+{
+	using namespace UE::DMXPixelMapping::Rendering::Preprocess::Private;
+	RenderInputProxy = MakeShared<FPreprocessRenderInputUserWidgetProxy>(InUserWidget, InputSize);
+	ApplyMaterialProxy = MakeShared<FPreprocessApplyFilterMaterialProxy>();
+}
+
+void UDMXPixelMappingPreprocessRenderer::ClearInput()
+{
+	RenderInputProxy.Reset();
+	ApplyMaterialProxy.Reset();
+}
+
+void UDMXPixelMappingPreprocessRenderer::Render()
+{
+	if (RenderInputProxy.IsValid())
+	{
+		RenderInputProxy->Render();
+		ApplyMaterialProxy->Render(RenderInputProxy->GetRenderedTexture(), *this);
+	}
+}
+
+UTexture* UDMXPixelMappingPreprocessRenderer::GetRenderedTexture() const
+{
+	return ApplyMaterialProxy->GetRenderedTexture();
+}
+
+TOptional<FVector2D> UDMXPixelMappingPreprocessRenderer::GetDesiredOutputSize2D() const
+{
+	if (OutputSizeMode == EDMXPixelMappingRenderingPreprocessorSizeMode::SameAsInput)
+	{
+		return RenderInputProxy->GetSize2D();
+	}
+	else if (OutputSizeMode == EDMXPixelMappingRenderingPreprocessorSizeMode::Downsampled)
+	{
+		return TOptional<FVector2D>();
+	}
+	else if (OutputSizeMode == EDMXPixelMappingRenderingPreprocessorSizeMode::CustomSize)
+	{
+		return TOptional<FVector2D>(CustomOutputSize);
+	}
+
+	return TOptional<FVector2D>();
+	checkf(0, TEXT("Unhandled output size mode in DMXPixelMappingRenderInputTextureProxy"));
+}
+
+FVector2D UDMXPixelMappingPreprocessRenderer::GetResultingSize2D() const
+{
+	if (OutputSizeMode == EDMXPixelMappingRenderingPreprocessorSizeMode::SameAsInput)
+	{
+		return RenderInputProxy->GetSize2D();
+	}
+	else if (OutputSizeMode == EDMXPixelMappingRenderingPreprocessorSizeMode::Downsampled)
+	{
+		const FVector2D DownsampledSize = RenderInputProxy->GetSize2D() / FMath::Pow(2.0, NumDownSamplePasses);
+		return FVector2D(FMath::Max(1.0, DownsampledSize.X), FMath::Max(1.0, DownsampledSize.Y));
+	}
+	else if (OutputSizeMode == EDMXPixelMappingRenderingPreprocessorSizeMode::CustomSize)
+	{
+		return CustomOutputSize;
+	}
+
+	checkf(0, TEXT("Unhandled output size mode in DMXPixelMappingRenderInputTextureProxy"));
+	return FVector2D::ZeroVector;
+}
+
+void UDMXPixelMappingPreprocessRenderer::PostLoad()
+{
+	Super::PostLoad();
+
+	if (FilterMaterial)
+	{
+		UPackage* TransientPackage = GetTransientPackage();
+		FilterMID = FilterMaterial ? UMaterialInstanceDynamic::Create(FilterMaterial, TransientPackage) : nullptr;
+	}
+}
+
+#if WITH_EDITOR
+void UDMXPixelMappingPreprocessRenderer::PostEditChangeProperty(FPropertyChangedEvent& PropertyChangedEvent)
+{
+	Super::PostEditChangeProperty(PropertyChangedEvent);
+
+	const FName PropertyName = PropertyChangedEvent.GetPropertyName();
+	if (PropertyName == GET_MEMBER_NAME_CHECKED(UDMXPixelMappingPreprocessRenderer, FilterMaterial))
+	{
+		if (FilterMaterial)
+		{
+			UPackage* TransientPackage = GetTransientPackage();
+			FilterMID = FilterMaterial ? UMaterialInstanceDynamic::Create(FilterMaterial, TransientPackage) : nullptr;
+		}
+		else
+		{
+			FilterMID = nullptr;
+		}
+	}
+}
+#endif // WITH_EDITOR
diff --git a/Engine/Plugins/VirtualProduction/DMX/DMXPixelMapping/Source/DMXPixelMappingRenderer/Private/DMXPixelMappingRenderer.cpp b/Engine/Plugins/VirtualProduction/DMX/DMXPixelMapping/Source/DMXPixelMappingRenderer/Private/DMXPixelMappingRenderer.cpp
--- a/Engine/Plugins/VirtualProduction/DMX/DMXPixelMapping/Source/DMXPixelMappingRenderer/Private/DMXPixelMappingRenderer.cpp
+++ b/Engine/Plugins/VirtualProduction/DMX/DMXPixelMapping/Source/DMXPixelMappingRenderer/Private/DMXPixelMappingRenderer.cpp
@@ -3,26 +3,23 @@
 #include "DMXPixelMappingRenderer.h"
 
 #include "Blueprint/UserWidget.h"
+#include "CanvasTypes.h"
 #include "ClearQuad.h"
 #include "CommonRenderResources.h"
+#include "Engine/Canvas.h"
 #include "Engine/TextureRenderTarget2D.h"
 #include "Materials/Material.h"
 #include "Modules/ModuleManager.h"
 #include "PixelShaderUtils.h"
-#include "ScreenRendering.h"
+#include "RenderingThread.h"
 #include "Rendering/SlateRenderer.h"
+#include "ScreenRendering.h"
+#include "TextureResource.h"
 #include "SlateMaterialBrush.h"
-#include "RenderingThread.h"
 #include "Slate/WidgetRenderer.h"
-#include "TextureResource.h"
 #include "Widgets/Images/SImage.h"
 #include "Widgets/SOverlay.h"
-#include "TextureResource.h"
-#include "Rendering/Texture2DResource.h"
-#include "Kismet/KismetRenderingLibrary.h"
-#include "CanvasTypes.h"
-#include "Engine/Canvas.h"
-#include "Materials/MaterialInstanceDynamic.h"
+
 
 namespace DMXPixelMappingRenderer
 {
@@ -98,38 +95,6 @@ FDMXPixelMappingRenderer::FDMXPixelMappingRenderer()
 	static const FName RendererModuleName("Renderer");
 	RendererModule = FModuleManager::GetModulePtr<IRendererModule>(RendererModuleName);
 
-	// Initialize Material Renderer
-	if (!MaterialWidgetRenderer.IsValid())
-	{
-		const bool bUseGammaCorrection = false;
-		MaterialWidgetRenderer = MakeShared<FWidgetRenderer>(bUseGammaCorrection);
-		check(MaterialWidgetRenderer.IsValid());
-	}
-
-	// Initialize Material Brush
-	if (!UIMaterialBrush.IsValid())
-	{
-		UIMaterialBrush = MakeShared<FSlateMaterialBrush>(FVector2D(1.f));
-		check(UIMaterialBrush.IsValid());
-	}
-
-	// Initialize UMG renderer
-	if (!UMGRenderer.IsValid())
-	{
-		const bool bUseGammaCorrection = true;
-		UMGRenderer = MakeShared<FWidgetRenderer>(bUseGammaCorrection);
-		check(UMGRenderer.IsValid());
-	}
-}
-
-void FDMXPixelMappingRenderer::PostProcessTexture(UTexture* InputTexture, const UE::DMXPixelMapping::FDMXPixelMappingInputTextureRenderingParameters& Params) const
-{
-	PostProcessProxy.Render(InputTexture, Params);
-}
-
-UTexture* FDMXPixelMappingRenderer::GetPostProcessedTexture() const
-{
-	return PostProcessProxy.GetRenderedTextureGameThread();
 }
 
 void FDMXPixelMappingRenderer::DownsampleRender(
@@ -241,6 +206,64 @@ void FDMXPixelMappingRenderer::DownsampleRender(
 		});
 }
 
+
+void FDMXPixelMappingRenderer::RenderMaterial(UTextureRenderTarget2D* InRenderTarget, UMaterialInterface* InMaterialInterface) const
+{	
+	// DEPRECATED 5.3
+	if (InMaterialInterface == nullptr)
+	{
+		return;
+	}
+
+	if (InRenderTarget == nullptr)
+	{
+		return;
+	}
+
+	FVector2D TextureSize = FVector2D(InRenderTarget->SizeX, InRenderTarget->SizeY);
+	UMaterial* Material = InMaterialInterface->GetMaterial();
+
+	if (Material != nullptr && Material->IsUIMaterial())
+	{
+		UIMaterialBrush->ImageSize = TextureSize;
+		UIMaterialBrush->SetMaterial(InMaterialInterface);
+
+		TSharedRef<SWidget> Widget =
+			SNew(SOverlay)
+			+ SOverlay::Slot()
+			[
+				SNew(SImage)
+				.Image(UIMaterialBrush.Get())
+			];
+
+		static const float DeltaTime = 0.f;
+		MaterialWidgetRenderer->DrawWidget(InRenderTarget, Widget, TextureSize, DeltaTime);
+		
+		// Reset material after drawing
+		UIMaterialBrush->SetMaterial(nullptr);
+	}
+}
+
+void FDMXPixelMappingRenderer::RenderWidget(UTextureRenderTarget2D* InRenderTarget, UUserWidget* InUserWidget) const
+{
+	// DEPRECATED 5.3
+	if (InUserWidget == nullptr)
+	{
+		return;
+	}
+
+	if (InRenderTarget == nullptr)
+	{
+		return;
+	}
+
+	FVector2D TextureSize = FVector2D(InRenderTarget->SizeX, InRenderTarget->SizeY);
+	static const float DeltaTime = 0.f;
+
+	UMGRenderer->DrawWidget(InRenderTarget, InUserWidget->TakeWidget(), TextureSize, DeltaTime);
+}
+
+
 #if WITH_EDITOR
 void FDMXPixelMappingRenderer::RenderPreview(const FTextureResource* TextureResource, const FTextureResource* DownsampleResource, TArray<FDMXPixelMappingDownsamplePixelPreviewParam>&& InPixelPreviewParamSet) const
 {
@@ -332,60 +355,6 @@ void FDMXPixelMappingRenderer::RenderPreview(const FTextureResource* TextureReso
 }
 #endif // WITH_EDITOR
 
-void FDMXPixelMappingRenderer::RenderMaterial(UTextureRenderTarget2D* InRenderTarget, UMaterialInterface* InMaterialInterface) const
-{
-	if (InMaterialInterface == nullptr)
-	{
-		return;
-	}
-
-	if (InRenderTarget == nullptr)
-	{
-		return;
-	}
-
-	FVector2D TextureSize = FVector2D(InRenderTarget->SizeX, InRenderTarget->SizeY);
-	UMaterial* Material = InMaterialInterface->GetMaterial();
-
-	if (Material != nullptr && Material->IsUIMaterial())
-	{
-		UIMaterialBrush->ImageSize = TextureSize;
-		UIMaterialBrush->SetMaterial(InMaterialInterface);
-
-		TSharedRef<SWidget> Widget =
-			SNew(SOverlay)
-			+ SOverlay::Slot()
-			[
-				SNew(SImage)
-				.Image(UIMaterialBrush.Get())
-			];
-
-		static const float DeltaTime = 0.f;
-		MaterialWidgetRenderer->DrawWidget(InRenderTarget, Widget, TextureSize, DeltaTime);
-		
-		// Reset material after drawing
-		UIMaterialBrush->SetMaterial(nullptr);
-	}
-}
-
-void FDMXPixelMappingRenderer::RenderWidget(UTextureRenderTarget2D* InRenderTarget, UUserWidget* InUserWidget) const
-{
-	if (InUserWidget == nullptr)
-	{
-		return;
-	}
-
-	if (InRenderTarget == nullptr)
-	{
-		return;
-	}
-
-	FVector2D TextureSize = FVector2D(InRenderTarget->SizeX, InRenderTarget->SizeY);
-	static const float DeltaTime = 0.f;
-
-	UMGRenderer->DrawWidget(InRenderTarget, InUserWidget->TakeWidget(), TextureSize, DeltaTime);
-}
-
 void FDMXPixelMappingRenderer::RenderTextureToRectangle(const FTextureResource* InTextureResource, const FTexture2DRHIRef InRenderTargetTexture, FVector2D InSize, bool bSRGBSource) const
 {
 	check(IsInGameThread());
diff --git a/Engine/Plugins/VirtualProduction/DMX/DMXPixelMapping/Source/DMXPixelMappingRenderer/Private/DMXPixelMappingRenderer.h b/Engine/Plugins/VirtualProduction/DMX/DMXPixelMapping/Source/DMXPixelMappingRenderer/Private/DMXPixelMappingRenderer.h
--- a/Engine/Plugins/VirtualProduction/DMX/DMXPixelMapping/Source/DMXPixelMappingRenderer/Private/DMXPixelMappingRenderer.h
+++ b/Engine/Plugins/VirtualProduction/DMX/DMXPixelMapping/Source/DMXPixelMappingRenderer/Private/DMXPixelMappingRenderer.h
@@ -4,7 +4,7 @@
 
 #include "IDMXPixelMappingRenderer.h"
 
-#include "DMXPixelMappingPostProcessProxy.h"
+#include "DMXPixelMappingPreprocessRenderer.h"
 
 class IRendererModule;
 class FWidgetRenderer;
@@ -24,12 +24,6 @@ public:
 	FDMXPixelMappingRenderer();
 
 	//~ Begin IDMXPixelMappingRenderer implementation
-
-	virtual void PostProcessTexture(UTexture* InputTexture, const UE::DMXPixelMapping::FDMXPixelMappingInputTextureRenderingParameters& Params) const override;
-
-	virtual UTexture* GetPostProcessedTexture() const override;
-
-	
 	virtual void DownsampleRender(
 		const FTextureResource* InputTexture,
 		const FTextureResource* DstTexture,
@@ -38,8 +32,14 @@ public:
 		DownsampleReadCallback InCallback
 	) const override;
 
+PRAGMA_DISABLE_DEPRECATION_WARNINGS
+	UE_DEPRECATED(5.3, "Use PreprocessRenderInputMaterialProxy instead")
 	virtual void RenderMaterial(UTextureRenderTarget2D* InRenderTarget, UMaterialInterface* InMaterialInterface) const override;
+
+	UE_DEPRECATED(5.3, "Use PreprocessRenderInputUserWidgetProxy instead")
 	virtual void RenderWidget(UTextureRenderTarget2D* InRenderTarget, UUserWidget* InUserWidget) const override;
+PRAGMA_ENABLE_DEPRECATION_WARNINGS
+
 	virtual void RenderTextureToRectangle(const FTextureResource* InTextureResource, const FTexture2DRHIRef InRenderTargetTexture, FVector2D InSize, bool bSRGBSource) const override;
 
 #if WITH_EDITOR
@@ -48,9 +48,6 @@ public:
 	//~ End IDMXPixelMappingRenderer implementation
 
 private:
-	/** Proxy to post process the input texture */
-	mutable UE::DMXPixelMapping::Renderer::Private::FDMXPixelMappingPostProccessProxy PostProcessProxy;
-
 	/** Brush for Material widget renderer */
 	TSharedPtr<FSlateMaterialBrush> UIMaterialBrush;
 
@@ -62,4 +59,4 @@ private:
 
 	/** The public interface of the renderer module. */
 	IRendererModule* RendererModule;
-};
\ No newline at end of file
+};
diff --git a/Engine/Plugins/VirtualProduction/DMX/DMXPixelMapping/Source/DMXPixelMappingRenderer/Private/PreprocessApplyFilterMaterialProxy.cpp b/Engine/Plugins/VirtualProduction/DMX/DMXPixelMapping/Source/DMXPixelMappingRenderer/Private/PreprocessApplyFilterMaterialProxy.cpp
new file mode 100644
--- /dev/null
+++ b/Engine/Plugins/VirtualProduction/DMX/DMXPixelMapping/Source/DMXPixelMappingRenderer/Private/PreprocessApplyFilterMaterialProxy.cpp
@@ -0,0 +1,278 @@
+// Copyright Epic Games, Inc. All Rights Reserved.
+
+#include "PreprocessApplyFilterMaterialProxy.h"
+
+#include "CanvasTypes.h"
+#include "CommonRenderResources.h"
+#include "DMXStats.h"
+#include "GlobalShader.h"
+#include "RHIStaticStates.h"
+#include "ScreenRendering.h"
+#include "TextureResource.h"
+#include "Engine/Canvas.h"
+#include "Engine/Engine.h"
+#include "Engine/Texture.h"
+#include "Engine/TextureRenderTarget2D.h"
+#include "Materials/MaterialInstanceDynamic.h"
+#include "Materials/MaterialInterface.h"
+#include "Materials/MaterialRenderProxy.h"
+#include "Modules/ModuleManager.h"
+
+
+namespace UE::DMXPixelMapping::Rendering::Preprocess::Private
+{
+	DECLARE_CYCLE_STAT(TEXT("PixelMapping ApplyFilterMaterial"), STAT_PreprocessApplyFilterMaterial, STATGROUP_DMX);
+
+	FDrawPreprocessMaterialCanvas::FDrawPreprocessMaterialCanvas()
+	{
+		Canvas = NewObject<UCanvas>();
+	}
+
+	void FDrawPreprocessMaterialCanvas::DrawMaterialToRenderTarget(UTextureRenderTarget2D* TextureRenderTarget, UMaterialInterface* Material)
+	{
+		if (!Material || !TextureRenderTarget || !TextureRenderTarget->GetResource())
+		{
+			return;
+		}
+
+		// This is a user-facing function, so we'd rather make sure that shaders are ready by the time we render, in order to ensure we don't draw with a fallback material
+		Material->EnsureIsComplete();
+		FTextureRenderTargetResource* RenderTargetResource = TextureRenderTarget->GameThread_GetRenderTargetResource();
+		if (!RenderTargetResource)
+		{
+			return;
+		}
+
+		FCanvas RenderCanvas(
+			RenderTargetResource,
+			nullptr,
+			GEngine->GetWorld(),
+			GMaxRHIFeatureLevel);
+
+		Canvas->Init(TextureRenderTarget->SizeX, TextureRenderTarget->SizeY, nullptr, &RenderCanvas);
+		{
+			ENQUEUE_RENDER_COMMAND(FlushDeferredResourceUpdateCommand)(
+				[RenderTargetResource](FRHICommandListImmediate& RHICmdList)
+				{
+					RenderTargetResource->FlushDeferredResourceUpdate(RHICmdList);
+				});
+
+			Canvas->K2_DrawMaterial(Material, FVector2D(0, 0), FVector2D(TextureRenderTarget->SizeX, TextureRenderTarget->SizeY), FVector2D(0, 0));
+
+			RenderCanvas.Flush_GameThread();
+			Canvas->Canvas = nullptr;
+
+			// UpdateResourceImmediate must be called here to ensure mips are generated.
+			TextureRenderTarget->UpdateResourceImmediate(false);
+
+			ENQUEUE_RENDER_COMMAND(ResetSceneTextureExtentHistory)(
+				[RenderTargetResource](FRHICommandListImmediate& RHICmdList)
+				{
+					RenderTargetResource->ResetSceneTextureExtentsHistory();
+				});
+		}
+	}
+
+	void FDrawPreprocessMaterialCanvas::AddReferencedObjects(FReferenceCollector& Collector)
+	{
+		Collector.AddReferencedObject(Canvas);
+	}
+
+	void FPreprocessApplyFilterMaterialProxy::Render(UTexture* InInputTexture, const UDMXPixelMappingPreprocessRenderer& InPreprocessRenderer)
+	{
+		SCOPE_CYCLE_COUNTER(STAT_PreprocessApplyFilterMaterial);
+
+		MaterialInstanceDynamic = InPreprocessRenderer.GetFilterMID();
+
+		WeakInputTexture = InInputTexture;
+		UpdateRenderTargets(InPreprocessRenderer.GetNumDownsamplePasses(), InPreprocessRenderer.GetDesiredOutputSize2D());
+
+		if (!CanRender())
+		{
+			return;
+		}
+
+		UTexture* InputTexture = WeakInputTexture.Get();
+		check(InputTexture);
+		if (MaterialInstanceDynamic && DownsampleRenderTargets.IsEmpty() && ensureMsgf(OutputRenderTarget, TEXT("Cannot apply pixel mapping filter material. Missing output render target.")))
+		{
+			// Render without downsampling
+			MaterialInstanceDynamic->SetTextureParameterValue(InPreprocessRenderer.GetInputTextureParameterName(), InputTexture);
+			MaterialInstanceDynamic->SetScalarParameterValue(InPreprocessRenderer.GetBlurDistanceParameterName(), InPreprocessRenderer.GetBlurDistance());
+			DrawMaterialCanvas.DrawMaterialToRenderTarget(OutputRenderTarget, MaterialInstanceDynamic);
+		}
+		else
+		{
+			const bool bApplyFilterMaterialEachDownsamplePass = InPreprocessRenderer.ShouldApplyFilterMaterialEachDownsamplePass();
+
+			// Render with downsampling. 
+			// Note, if the last downsample render target is of minimal size, further renderer targets are not created.
+			// Hence use the number of downsample render targets instead of InParams.NumDownsamplePasses.
+			const int32 NumDownsamplePasses = DownsampleRenderTargets.Num();
+			for (int32 DownsamplePass = 0; DownsamplePass < NumDownsamplePasses; DownsamplePass++)
+			{
+				UTexture* SourceTexture = DownsamplePass == 0 ? InputTexture : DownsampleRenderTargets[DownsamplePass - 1];
+				UTextureRenderTarget2D* DownsampleRenderTarget = DownsampleRenderTargets[DownsamplePass].Get();
+
+				const bool bApplyFilterMaterial = MaterialInstanceDynamic && (bApplyFilterMaterialEachDownsamplePass || DownsampleRenderTarget == DownsampleRenderTargets.Last());
+				if (bApplyFilterMaterial)
+				{
+					MaterialInstanceDynamic->SetTextureParameterValue(InPreprocessRenderer.GetInputTextureParameterName(), SourceTexture);
+					MaterialInstanceDynamic->SetScalarParameterValue(InPreprocessRenderer.GetBlurDistanceParameterName(), InPreprocessRenderer.GetBlurDistance());
+					DrawMaterialCanvas.DrawMaterialToRenderTarget(DownsampleRenderTarget, MaterialInstanceDynamic);
+				}
+				else
+				{
+					RenderTextureToTarget(SourceTexture, DownsampleRenderTarget);
+				}
+			}
+
+			// Scale to output size
+			UTexture* LastSourceTexture = DownsampleRenderTargets.IsEmpty() ? InputTexture : DownsampleRenderTargets.Last();
+			RenderTextureToTarget(LastSourceTexture, OutputRenderTarget);
+		}
+	}
+
+	UTexture* FPreprocessApplyFilterMaterialProxy::GetRenderedTexture() const
+	{
+		return CanRender() ? OutputRenderTarget : WeakInputTexture.Get();
+	}
+
+	void FPreprocessApplyFilterMaterialProxy::AddReferencedObjects(FReferenceCollector& Collector)
+	{
+		Collector.AddReferencedObject(MaterialInstanceDynamic);
+		Collector.AddReferencedObjects(DownsampleRenderTargets);
+		Collector.AddReferencedObject(OutputRenderTarget);
+	}
+
+	bool FPreprocessApplyFilterMaterialProxy::CanRender() const
+	{
+		return WeakInputTexture.IsValid() && OutputRenderTarget && FApp::CanEverRender();
+	}
+
+	void FPreprocessApplyFilterMaterialProxy::UpdateRenderTargets(int32 NumDownsamplePasses, const TOptional<FVector2D>& OptionalOutputSize)
+	{
+		UTexture* InputTexture = WeakInputTexture.Get();
+
+		const bool bResize = 
+			OptionalOutputSize.IsSet() &&
+			InputTexture &&
+			InputTexture->GetSurfaceWidth() != OptionalOutputSize.GetValue().X &&
+			InputTexture->GetSurfaceHeight() != OptionalOutputSize.GetValue().Y;
+
+		const bool bNeedsAnyRenderTargets = WeakInputTexture.IsValid() && (NumDownsamplePasses > 0 || MaterialInstanceDynamic || bResize);
+		if (!bNeedsAnyRenderTargets)
+		{
+			OutputRenderTarget = nullptr;
+			return;
+		}
+
+		const int32 NumExpectedDownsampleRenderTargets = OptionalOutputSize.IsSet() ? NumDownsamplePasses : NumDownsamplePasses - 1;
+		const bool bInvalidRenderTargets = 
+			!OutputRenderTarget || 
+			NumExpectedDownsampleRenderTargets != DownsampleRenderTargets.Num() ||
+			(OptionalOutputSize.IsSet() && OutputRenderTarget->GetSurfaceWidth() != OptionalOutputSize.GetValue().X) ||
+			(OptionalOutputSize.IsSet() && OutputRenderTarget->GetSurfaceHeight() != OptionalOutputSize.GetValue().Y);
+
+		if (bInvalidRenderTargets)
+		{
+			FlushRenderingCommands();
+
+			const FVector2D InputSize = FVector2D(InputTexture->GetSurfaceHeight(), InputTexture->GetSurfaceWidth());
+
+			// Create downsample render targets
+			DownsampleRenderTargets.Reset();
+			FVector2D DownsampleSize = InputSize;
+			for (int32 DownsamplePass = 0; DownsamplePass < NumDownsamplePasses; DownsamplePass++)
+			{
+				DownsampleSize /= 2.0;
+				if (DownsampleSize.X <= 1.0 || DownsampleSize.Y <= 1.0)
+				{
+					break;
+				}
+				UTextureRenderTarget2D* DownsampleRenderTarget = NewObject<UTextureRenderTarget2D>();
+				DownsampleRenderTarget->ClearColor = FLinearColor::Black;
+				DownsampleRenderTarget->InitAutoFormat(DownsampleSize.X, DownsampleSize.Y);
+				DownsampleRenderTarget->UpdateResourceImmediate();
+				DownsampleRenderTargets.Add(DownsampleRenderTarget);
+			}
+
+			UTextureRenderTarget2D* ScaleRenderTarget = nullptr;
+			if (OptionalOutputSize.IsSet())
+			{
+				ScaleRenderTarget = NewObject<UTextureRenderTarget2D>();
+				ScaleRenderTarget->ClearColor = FLinearColor::Black;
+				ScaleRenderTarget->InitAutoFormat(OptionalOutputSize.GetValue().X, OptionalOutputSize.GetValue().Y);
+				ScaleRenderTarget->UpdateResourceImmediate();
+			}
+
+			if (ScaleRenderTarget)
+			{
+				OutputRenderTarget = ScaleRenderTarget;
+			}
+			else if (!DownsampleRenderTargets.IsEmpty())
+			{
+				OutputRenderTarget = DownsampleRenderTargets.Pop();
+			}
+			else
+			{
+				OutputRenderTarget = nullptr;
+			}
+		}
+	}
+
+	void FPreprocessApplyFilterMaterialProxy::RenderTextureToTarget(UTexture* Texture, UTextureRenderTarget2D* RenderTarget) const
+	{
+		IRendererModule* RendererModule = FModuleManager::GetModulePtr<IRendererModule>("Renderer");
+		ENQUEUE_RENDER_COMMAND(PixelMappingPreprocessDownsamplePass)(
+			[RendererModule, Texture, RenderTarget, this](FRHICommandListImmediate& RHICmdList)
+			{
+				if (!Texture || !Texture->GetResource())
+				{
+					return;
+				}
+
+				const FTextureRHIRef SourceTextureRHI = Texture->GetResource()->GetTexture2DRHI();
+				const FTextureRHIRef TargetTextureRHI = RenderTarget->GetResource()->GetTexture2DRHI();
+
+				FRHIRenderPassInfo RPInfo(TargetTextureRHI, MakeRenderTargetActions(ERenderTargetLoadAction::ENoAction, ERenderTargetStoreAction::EStore));
+
+				RHICmdList.Transition(FRHITransitionInfo(TargetTextureRHI, ERHIAccess::Unknown, ERHIAccess::RTV));
+				RHICmdList.BeginRenderPass(RPInfo, TEXT("PixelMappingFilterProxy_RenderToTarget"));
+				{
+					RHICmdList.SetViewport(0, 0, 0.0f, (float)TargetTextureRHI->GetSizeX(), (float)TargetTextureRHI->GetSizeY(), 1.0f);
+
+					FGraphicsPipelineStateInitializer GraphicsPSOInit;
+					RHICmdList.ApplyCachedRenderTargets(GraphicsPSOInit);
+					GraphicsPSOInit.BlendState = TStaticBlendState<>::GetRHI();
+					GraphicsPSOInit.RasterizerState = TStaticRasterizerState<>::GetRHI();
+					GraphicsPSOInit.DepthStencilState = TStaticDepthStencilState<false, CF_Always>::GetRHI();
+
+					FGlobalShaderMap* ShaderMap = GetGlobalShaderMap(GMaxRHIFeatureLevel);
+					TShaderMapRef<FScreenVS> VertexShader(ShaderMap);
+					TShaderMapRef<FScreenPS> PixelShader(ShaderMap);
+
+					GraphicsPSOInit.BoundShaderState.VertexDeclarationRHI = GFilterVertexDeclaration.VertexDeclarationRHI;
+					GraphicsPSOInit.BoundShaderState.VertexShaderRHI = VertexShader.GetVertexShader();
+					GraphicsPSOInit.BoundShaderState.PixelShaderRHI = PixelShader.GetPixelShader();
+					GraphicsPSOInit.PrimitiveType = PT_TriangleList;
+
+					SetGraphicsPipelineState(RHICmdList, GraphicsPSOInit, 0);
+
+					SetShaderParametersLegacyPS(RHICmdList, PixelShader, TStaticSamplerState<SF_Bilinear>::GetRHI(), SourceTextureRHI);
+
+					RendererModule->DrawRectangle(RHICmdList,
+						0, 0,									// Dest X, Y
+						(float)TargetTextureRHI->GetSizeX(),	// Dest Width
+						(float)TargetTextureRHI->GetSizeY(),	// Dest Height
+						0, 0,									// Source U, V
+						1, 1,									// Source USize, VSize
+						TargetTextureRHI->GetSizeXY(),			// Target buffer size
+						FIntPoint(1, 1),						// Source texture size
+						VertexShader,
+						EDRF_Default);
+				}
+				RHICmdList.EndRenderPass();
+			});
+	}
+}
diff --git a/Engine/Plugins/VirtualProduction/DMX/DMXPixelMapping/Source/DMXPixelMappingRenderer/Private/PreprocessApplyFilterMaterialProxy.h b/Engine/Plugins/VirtualProduction/DMX/DMXPixelMapping/Source/DMXPixelMappingRenderer/Private/PreprocessApplyFilterMaterialProxy.h
new file mode 100644
--- /dev/null
+++ b/Engine/Plugins/VirtualProduction/DMX/DMXPixelMapping/Source/DMXPixelMappingRenderer/Private/PreprocessApplyFilterMaterialProxy.h
@@ -0,0 +1,89 @@
+// Copyright Epic Games, Inc. All Rights Reserved.
+
+#pragma once
+
+#include "DMXPixelMappingPreprocessRenderer.h"
+
+#include "UObject/GCObject.h"
+#include "UObject/ObjectPtr.h"
+
+class UCanvas;
+class UMaterialInstanceDynamic;
+class UMaterialInterface;
+class UTexture;
+class UTextureRenderTarget2D;
+
+
+namespace UE::DMXPixelMapping::Rendering::Preprocess::Private
+{
+	/** Canvas implementation for pixelmapping */
+	class FDrawPreprocessMaterialCanvas
+		: public FGCObject
+	{
+	public:
+		/** Constructor */
+		FDrawPreprocessMaterialCanvas();
+
+		/** Draws a material to a render target */
+		void DrawMaterialToRenderTarget(UTextureRenderTarget2D* TextureRenderTarget, UMaterialInterface* Material);
+
+	protected:
+		//~ Begin FGCObject interface
+		virtual void AddReferencedObjects(FReferenceCollector& Collector) override;
+		virtual FString GetReferencerName() const override
+		{
+			return TEXT("DMXPixelMapping::Rendering::FDrawPreprocessMaterialCanvas");
+		}
+		//~ End FGCObject interface
+
+	private:
+		/** Canvas used to draw the material */
+		TObjectPtr<UCanvas> Canvas;
+	};
+
+	/** Proxy to filter the texture used in Pixel Mapping */
+	class FPreprocessApplyFilterMaterialProxy
+		: public IPreprocessApplyFilterMaterialProxy
+		, public FGCObject
+	{
+	public:
+		//~ Begin IPreprocessApplyFilterMaterialProxy interface
+		virtual void Render(UTexture* InInputTexture, const UDMXPixelMappingPreprocessRenderer& InPreprocessRenderer) override;
+		virtual UTexture* GetRenderedTexture() const override;
+		//~ End IPreprocessApplyFilterMaterialProxy interface
+
+	protected:
+		//~ Begin FGCObject interface
+		virtual void AddReferencedObjects(FReferenceCollector& Collector) override;
+		virtual FString GetReferencerName() const override
+		{
+			return TEXT("FPreprocessApplyFilterMaterialProxy");
+		}
+		//~ End FGCObject interface
+
+	private:
+		/** Returns true if rendering is needed */
+		bool CanRender() const;
+
+		/** Updates render targets */
+		void UpdateRenderTargets(int32 NumDownsamplePasses, const TOptional<FVector2D>& OptionalOutputSize);
+
+		/** Renders the Input Texture to the Output render target*/
+		void RenderTextureToTarget(UTexture* Texture, UTextureRenderTarget2D* RenderTarget) const;
+
+		/** Weak ref to the input texture */
+		TWeakObjectPtr<UTexture> WeakInputTexture;
+
+		/** The dynamic material instance to be applied to the input texture */
+		TObjectPtr<UMaterialInstanceDynamic> MaterialInstanceDynamic;
+
+		/** Downsample render targets */
+		TArray<TObjectPtr<UTextureRenderTarget2D>> DownsampleRenderTargets;
+
+		/** The target to which the output is rendered. Either the last downscale render target or the scale render target */
+		TObjectPtr<UTextureRenderTarget2D> OutputRenderTarget;
+
+		/** Canvas used to draw materials */
+		FDrawPreprocessMaterialCanvas DrawMaterialCanvas;
+	};
+}
diff --git a/Engine/Plugins/VirtualProduction/DMX/DMXPixelMapping/Source/DMXPixelMappingRenderer/Private/PreprocessRenderInputMaterialProxy .cpp b/Engine/Plugins/VirtualProduction/DMX/DMXPixelMapping/Source/DMXPixelMappingRenderer/Private/PreprocessRenderInputMaterialProxy .cpp
new file mode 100644
--- /dev/null
+++ b/Engine/Plugins/VirtualProduction/DMX/DMXPixelMapping/Source/DMXPixelMappingRenderer/Private/PreprocessRenderInputMaterialProxy .cpp	
@@ -0,0 +1,82 @@
+// Copyright Epic Games, Inc. All Rights Reserved.
+
+#include "PreprocessRenderInputMaterialProxy.h"
+
+#include "DMXPixelMappingPreprocessRenderer.h"
+#include "DMXStats.h"
+#include "Engine/TextureRenderTarget2D.h"
+#include "Materials/Material.h"
+#include "Materials/MaterialInterface.h"
+#include "SlateMaterialBrush.h"
+#include "Slate/WidgetRenderer.h"
+#include "Widgets/Images/SImage.h"
+
+
+namespace UE::DMXPixelMapping::Rendering::Preprocess::Private
+{
+	DECLARE_CYCLE_STAT(TEXT("PixelMapping RenderInputMaterial"), STAT_PreprocessRenderInputMaterial, STATGROUP_DMX);
+
+	FPreprocessRenderInputMaterialProxy::FPreprocessRenderInputMaterialProxy(UMaterialInterface* InMaterial, const FVector2D& InInputSize)
+		: WeakMaterial(InMaterial)
+	{
+		if (InMaterial)
+		{
+			InMaterial->EnsureIsComplete();
+		}
+
+		IntermediateRenderTarget = NewObject<UTextureRenderTarget2D>();
+		IntermediateRenderTarget->ClearColor = FLinearColor::Black;
+		IntermediateRenderTarget->InitAutoFormat(InInputSize.X, InInputSize.Y);
+		IntermediateRenderTarget->UpdateResourceImmediate();
+
+		UIMaterialBrush = MakeShared<FSlateMaterialBrush>(FVector2D(1.f));
+
+		const bool bUseGammaCorrection = false;
+		MaterialRenderer = MakeShared<FWidgetRenderer>(bUseGammaCorrection);
+	}
+
+	void FPreprocessRenderInputMaterialProxy::Render()
+	{
+		SCOPE_CYCLE_COUNTER(STAT_PreprocessRenderInputMaterial);
+
+		UMaterial* Material = WeakMaterial.IsValid() ? WeakMaterial->GetMaterial() : nullptr;
+		if (Material && Material->IsUIMaterial())
+		{
+			check(IntermediateRenderTarget);
+			FVector2D TextureSize = FVector2D(IntermediateRenderTarget->SizeX, IntermediateRenderTarget->SizeY);
+
+			UIMaterialBrush->ImageSize = TextureSize;
+			UIMaterialBrush->SetMaterial(Material);
+
+			TSharedRef<SWidget> Widget =
+				SNew(SOverlay)
+				+ SOverlay::Slot()
+				[
+					SNew(SImage)
+					.Image(UIMaterialBrush.Get())
+				];
+
+			static const float DeltaTime = 0.f; 
+			MaterialRenderer->DrawWidget(IntermediateRenderTarget, Widget, TextureSize, DeltaTime);
+
+			// Reset material after drawing
+			UIMaterialBrush->SetMaterial(nullptr);
+		}
+	}
+
+	UTexture* FPreprocessRenderInputMaterialProxy::GetRenderedTexture() const
+	{
+		return IntermediateRenderTarget;
+	}
+
+	FVector2D FPreprocessRenderInputMaterialProxy::GetSize2D() const
+	{
+		check(IntermediateRenderTarget);
+		return FVector2D(IntermediateRenderTarget->GetSurfaceWidth(), IntermediateRenderTarget->GetSurfaceHeight());
+	}
+
+	void FPreprocessRenderInputMaterialProxy::AddReferencedObjects(FReferenceCollector& Collector)
+	{
+		Collector.AddReferencedObject(IntermediateRenderTarget);
+	}
+}
diff --git a/Engine/Plugins/VirtualProduction/DMX/DMXPixelMapping/Source/DMXPixelMappingRenderer/Private/PreprocessRenderInputMaterialProxy.h b/Engine/Plugins/VirtualProduction/DMX/DMXPixelMapping/Source/DMXPixelMappingRenderer/Private/PreprocessRenderInputMaterialProxy.h
new file mode 100644
--- /dev/null
+++ b/Engine/Plugins/VirtualProduction/DMX/DMXPixelMapping/Source/DMXPixelMappingRenderer/Private/PreprocessRenderInputMaterialProxy.h
@@ -0,0 +1,53 @@
+// Copyright Epic Games, Inc. All Rights Reserved.
+
+#pragma once
+
+#include "DMXPixelMappingPreprocessRenderer.h"
+
+#include "UObject/GCObject.h"
+
+struct FSlateMaterialBrush;
+class FWidgetRenderer;
+class UMaterialInterface;
+class UTextureRenderTarget2D;
+
+
+namespace UE::DMXPixelMapping::Rendering::Preprocess::Private
+{
+	/** Proxy responsible for the input material */
+	class FPreprocessRenderInputMaterialProxy
+		: public IPreprocessRenderInputProxy
+		, public FGCObject
+	{
+	public:
+		FPreprocessRenderInputMaterialProxy(UMaterialInterface* InMaterial, const FVector2D& InInputSize);
+
+		//~ Begin IPreprocessRenderInputProxy interface
+		virtual void Render() override;
+		virtual UTexture* GetRenderedTexture() const override;
+		virtual FVector2D GetSize2D() const override;
+		//~ End IPreprocessRenderInputProxy interface
+
+	protected:
+		//~ Begin FGCObject interface
+		virtual void AddReferencedObjects(FReferenceCollector& Collector) override;
+		virtual FString GetReferencerName() const override
+		{
+			return TEXT("DMXPixelMapping::Rendering::Private::FPreprocessRenderInputMaterialProxy");
+		}
+		//~ End FGCObject interface
+
+	private:
+		/** Material renderer */
+		TSharedPtr<FWidgetRenderer> MaterialRenderer;
+
+		/** Brush for Material widget renderer */
+		TSharedPtr<FSlateMaterialBrush> UIMaterialBrush;
+
+		/** The Input object */
+		TWeakObjectPtr<UMaterialInterface> WeakMaterial;
+
+		/** Intermediate render target for materials and user widgets */
+		TObjectPtr<UTextureRenderTarget2D> IntermediateRenderTarget;
+	};
+}
diff --git a/Engine/Plugins/VirtualProduction/DMX/DMXPixelMapping/Source/DMXPixelMappingRenderer/Private/PreprocessRenderInputTextureProxy.cpp b/Engine/Plugins/VirtualProduction/DMX/DMXPixelMapping/Source/DMXPixelMappingRenderer/Private/PreprocessRenderInputTextureProxy.cpp
new file mode 100644
--- /dev/null
+++ b/Engine/Plugins/VirtualProduction/DMX/DMXPixelMapping/Source/DMXPixelMappingRenderer/Private/PreprocessRenderInputTextureProxy.cpp
@@ -0,0 +1,35 @@
+// Copyright Epic Games, Inc. All Rights Reserved.
+
+#include "PreprocessRenderInputTextureProxy.h"
+
+#include "Engine/Texture.h"
+
+
+namespace UE::DMXPixelMapping::Rendering::Preprocess::Private
+{
+	FPreprocessRenderInputTextureProxy::FPreprocessRenderInputTextureProxy(UTexture* InTexture)
+		:  WeakInputTexture(InTexture)
+	{
+		if (InTexture)
+		{
+			InTexture->WaitForStreaming();
+		}
+	}
+
+	void FPreprocessRenderInputTextureProxy::Render()
+	{
+		// No rendering required
+	}
+
+	UTexture* FPreprocessRenderInputTextureProxy::GetRenderedTexture() const
+	{
+		return WeakInputTexture.Get();
+	}
+
+	FVector2D FPreprocessRenderInputTextureProxy::GetSize2D() const
+	{
+		return WeakInputTexture.IsValid() ?
+			FVector2D(WeakInputTexture->GetSurfaceWidth(), WeakInputTexture->GetSurfaceHeight()) :
+			FVector2D{ 0.0, 0.0 };
+	}
+}
diff --git a/Engine/Plugins/VirtualProduction/DMX/DMXPixelMapping/Source/DMXPixelMappingRenderer/Private/PreprocessRenderInputTextureProxy.h b/Engine/Plugins/VirtualProduction/DMX/DMXPixelMapping/Source/DMXPixelMappingRenderer/Private/PreprocessRenderInputTextureProxy.h
new file mode 100644
--- /dev/null
+++ b/Engine/Plugins/VirtualProduction/DMX/DMXPixelMapping/Source/DMXPixelMappingRenderer/Private/PreprocessRenderInputTextureProxy.h
@@ -0,0 +1,33 @@
+// Copyright Epic Games, Inc. All Rights Reserved.
+
+#pragma once
+
+#include "DMXPixelMappingPreprocessRenderer.h"
+
+#include "UObject/WeakObjectPtr.h"
+
+class IDMXPixelMappingRenderer;
+class UDMXPixelMappingPreprocessRenderer;
+class UTexture;
+
+
+namespace UE::DMXPixelMapping::Rendering::Preprocess::Private
+{
+	/** Proxy responsible for the input texture */
+	class FPreprocessRenderInputTextureProxy
+		: public IPreprocessRenderInputProxy
+	{
+	public:
+		FPreprocessRenderInputTextureProxy(UTexture* InTexture);
+
+		//~ Begin IPreprocessRenderInputProxy interface
+		virtual void Render() override;
+		virtual UTexture* GetRenderedTexture() const override;
+		virtual FVector2D GetSize2D() const override;
+		//~ End IPreprocessRenderInputProxy interface
+
+	private:
+		/** The Input object */
+		TWeakObjectPtr<UTexture> WeakInputTexture;
+	};
+}
diff --git a/Engine/Plugins/VirtualProduction/DMX/DMXPixelMapping/Source/DMXPixelMappingRenderer/Private/PreprocessRenderInputUserWidgetProxy.cpp b/Engine/Plugins/VirtualProduction/DMX/DMXPixelMapping/Source/DMXPixelMappingRenderer/Private/PreprocessRenderInputUserWidgetProxy.cpp
new file mode 100644
--- /dev/null
+++ b/Engine/Plugins/VirtualProduction/DMX/DMXPixelMapping/Source/DMXPixelMappingRenderer/Private/PreprocessRenderInputUserWidgetProxy.cpp
@@ -0,0 +1,56 @@
+// Copyright Epic Games, Inc. All Rights Reserved.
+
+#include "PreprocessRenderInputUserWidgetProxy.h"
+
+#include "Blueprint/UserWidget.h"
+#include "DMXStats.h"
+#include "Engine/TextureRenderTarget2D.h"
+#include "Slate/WidgetRenderer.h"
+
+
+namespace UE::DMXPixelMapping::Rendering::Preprocess::Private
+{
+	DECLARE_CYCLE_STAT(TEXT("PixelMapping RenderInputUserWidget"), STAT_PreprocessRenderInputUserWidget, STATGROUP_DMX);
+
+	FPreprocessRenderInputUserWidgetProxy::FPreprocessRenderInputUserWidgetProxy(UUserWidget* InUserWidget, const FVector2D& InInputSize)
+		: WeakUserWidget(InUserWidget)
+	{
+		IntermediateRenderTarget = NewObject<UTextureRenderTarget2D>();
+		IntermediateRenderTarget->ClearColor = FLinearColor::Black;
+		IntermediateRenderTarget->InitAutoFormat(InInputSize.X, InInputSize.Y);
+		IntermediateRenderTarget->UpdateResourceImmediate();
+
+		const bool bUseGammaCorrection = false;
+		UMGRenderer = MakeShared<FWidgetRenderer>(bUseGammaCorrection);
+	}
+
+	void FPreprocessRenderInputUserWidgetProxy::Render()
+	{
+		SCOPE_CYCLE_COUNTER(STAT_PreprocessRenderInputUserWidget);
+
+		if (UUserWidget* UserWidget = WeakUserWidget.Get())
+		{
+			check(IntermediateRenderTarget);
+			FVector2D TextureSize = FVector2D(IntermediateRenderTarget->SizeX, IntermediateRenderTarget->SizeY);
+
+			static const float DeltaTime = 0.f;
+			UMGRenderer->DrawWidget(IntermediateRenderTarget, UserWidget->TakeWidget(), TextureSize, DeltaTime);
+		}
+	}
+
+	UTexture* FPreprocessRenderInputUserWidgetProxy::GetRenderedTexture() const
+	{
+		return IntermediateRenderTarget;
+	}
+
+	FVector2D FPreprocessRenderInputUserWidgetProxy::GetSize2D() const
+	{
+		check(IntermediateRenderTarget);
+		return FVector2D(IntermediateRenderTarget->GetSurfaceWidth(), IntermediateRenderTarget->GetSurfaceHeight());
+	}
+
+	void FPreprocessRenderInputUserWidgetProxy::AddReferencedObjects(FReferenceCollector& Collector)
+	{
+		Collector.AddReferencedObject(IntermediateRenderTarget);
+	}
+}
diff --git a/Engine/Plugins/VirtualProduction/DMX/DMXPixelMapping/Source/DMXPixelMappingRenderer/Private/PreprocessRenderInputUserWidgetProxy.h b/Engine/Plugins/VirtualProduction/DMX/DMXPixelMapping/Source/DMXPixelMappingRenderer/Private/PreprocessRenderInputUserWidgetProxy.h
new file mode 100644
--- /dev/null
+++ b/Engine/Plugins/VirtualProduction/DMX/DMXPixelMapping/Source/DMXPixelMappingRenderer/Private/PreprocessRenderInputUserWidgetProxy.h
@@ -0,0 +1,49 @@
+// Copyright Epic Games, Inc. All Rights Reserved.
+
+#pragma once
+
+#include "DMXPixelMappingPreprocessRenderer.h"
+
+#include "UObject/GCObject.h"
+
+class FWidgetRenderer;
+class UUserWidget;
+class UTextureRenderTarget2D;
+
+
+namespace UE::DMXPixelMapping::Rendering::Preprocess::Private
+{
+	/** Proxy responsible for the input user widget */
+	class FPreprocessRenderInputUserWidgetProxy
+		: public IPreprocessRenderInputProxy
+		, public FGCObject
+	{
+	public:
+		FPreprocessRenderInputUserWidgetProxy(UUserWidget* InUserWidget, const FVector2D& InInputSize);
+
+		//~ Begin IPreprocessRenderInputProxy interface
+		virtual void Render() override;
+		virtual UTexture* GetRenderedTexture() const override;
+		virtual FVector2D GetSize2D() const override;
+		//~ End IPreprocessRenderInputProxy interface
+
+	protected:
+		//~ Begin FGCObject interface
+		virtual void AddReferencedObjects(FReferenceCollector& Collector) override;
+		virtual FString GetReferencerName() const override
+		{
+			return TEXT("DMXPixelMapping::Rendering::Private::FPreprocessRenderInputUserWidgetProxy");
+		}
+		//~ End FGCObject interface
+
+	private:
+		/** UMG widget renderer */
+		TSharedPtr<FWidgetRenderer> UMGRenderer;
+
+		/** The Input object */
+		TWeakObjectPtr<UUserWidget> WeakUserWidget;
+
+		/** Intermediate render target for materials and user widgets */
+		TObjectPtr<UTextureRenderTarget2D> IntermediateRenderTarget;
+	};
+}
diff --git a/Engine/Plugins/VirtualProduction/DMX/DMXPixelMapping/Source/DMXPixelMappingRenderer/Public/DMXPixelMappingPreprocessRenderer.h b/Engine/Plugins/VirtualProduction/DMX/DMXPixelMapping/Source/DMXPixelMappingRenderer/Public/DMXPixelMappingPreprocessRenderer.h
new file mode 100644
--- /dev/null
+++ b/Engine/Plugins/VirtualProduction/DMX/DMXPixelMapping/Source/DMXPixelMappingRenderer/Public/DMXPixelMappingPreprocessRenderer.h
@@ -0,0 +1,168 @@
+// Copyright Epic Games, Inc. All Rights Reserved.
+
+#pragma once
+
+#include "UObject/Object.h"
+
+#include "DMXPixelMappingPreprocessRenderer.generated.h"
+
+class IDMXPixelMappingRenderer;
+class UMaterialInstanceDynamic;
+class UMaterialInterface;
+class UTexture;
+class UUserWidget;
+
+
+UENUM(BlueprintType)
+enum class EDMXPixelMappingRenderingPreprocessorSizeMode : uint8
+{
+	SameAsInput,
+	Downsampled,
+	CustomSize	
+};
+
+namespace UE::DMXPixelMapping::Rendering::Preprocess::Private
+{
+	/** Interface for the object responsible to render the input texture/material/umg widget. */
+	class IPreprocessRenderInputProxy
+		: public TSharedFromThis<IPreprocessRenderInputProxy>
+	{
+	public:
+		virtual ~IPreprocessRenderInputProxy() {}
+
+		/** Renders the input texture, optionally using a Filter material */
+		virtual void Render() = 0;
+
+		/** Returns the currently rendered texture, or nullptr if no texture was rendered yet */
+		virtual UTexture* GetRenderedTexture() const = 0;
+
+		/** Returns the size of the input object */
+		virtual FVector2D GetSize2D() const = 0;
+	};
+
+	/** Interface for the object responsible to apply the material to the input texture. */
+	class IPreprocessApplyFilterMaterialProxy
+		: public TSharedFromThis<IPreprocessApplyFilterMaterialProxy>
+	{
+	public:
+		virtual ~IPreprocessApplyFilterMaterialProxy() {}
+
+		/** Renders the input texture, optionally using a Filter material */
+		virtual void Render(UTexture* InInputTexture, const class UDMXPixelMappingPreprocessRenderer& InPreprocessRenderer) = 0;
+
+		/** Returns the preprocessed texture */
+		virtual UTexture* GetRenderedTexture() const = 0;
+	};
+}
+
+/** Renders the texture that is used to pixel map */
+UCLASS(BlueprintType)
+class DMXPIXELMAPPINGRENDERER_API UDMXPixelMappingPreprocessRenderer
+	: public UObject
+{
+	GENERATED_BODY()
+
+public:
+	/** Sets the input texture that is used when rendering */
+	void SetInputTexture(UTexture* InTexture);
+
+	/** Sets the input texture that is used when rendering */
+	void SetInputMaterial(UMaterialInterface* InMaterial);
+
+	/** Sets the input texture that is used when rendering */
+	void SetInputUserWidget(UUserWidget* InUserWidget);
+
+	/** Clears input */
+	void ClearInput();
+
+	/** Renders the current input */
+	void Render();
+
+	/** Returns the rendered input texture, or nullptr if no texture was rendered yet */
+	UTexture* GetRenderedTexture() const;
+
+	/** Returns the resulting size of preprocessing, even if rendering did not occur yet */
+	FVector2D GetResultingSize2D() const;
+
+	/** Returns the Filter material, or nullptr if no Filter material is set */
+	const TObjectPtr<UMaterialInstanceDynamic>& GetFilterMID() const { return FilterMID; }
+
+	/** Returns true, if the Filter material should be applied each downsample pass */
+	bool ShouldApplyFilterMaterialEachDownsamplePass() const { return bApplyFilterMaterialEachDownsamplePass; }
+
+	/** Returns the number of time the texture should be downsampled by factor 2 */
+	int32 GetNumDownsamplePasses() const { return NumDownSamplePasses; }
+
+	/** Returns the desired output size, or an unset optional vector 2D if the resulting size should be used */
+	TOptional<FVector2D> GetDesiredOutputSize2D() const;
+
+	/** Returns the blur distance parameter */
+	float GetBlurDistance() const { return BlurDistance; }
+
+	/** Returns the input texture material parameter name */
+	FName GetInputTextureParameterName() const { return InputTextureParameterName; }
+
+	/** Returns the blur distance material parameter name */
+	FName GetBlurDistanceParameterName() const { return BlurDistanceParamterName; }
+
+	// Property name getters
+	static FName GetOutputSizeModeMemberNameChecked() { return GET_MEMBER_NAME_CHECKED(UDMXPixelMappingPreprocessRenderer, OutputSizeMode); }
+	static FName GetCustomOutputSizeMemberNameChecked() { return GET_MEMBER_NAME_CHECKED(UDMXPixelMappingPreprocessRenderer, CustomOutputSize); }
+	static FName GetFilterMaterialMemberNameChecked() { return GET_MEMBER_NAME_CHECKED(UDMXPixelMappingPreprocessRenderer, FilterMaterial); }
+	static FName GetBlurDistanceMemberNameChecked() { return GET_MEMBER_NAME_CHECKED(UDMXPixelMappingPreprocessRenderer, BlurDistance); }
+
+protected:
+	//~ Begin UObject interface
+	virtual void PostLoad() override;
+#if WITH_EDITOR
+	virtual void PostEditChangeProperty(FPropertyChangedEvent& PropertyChangedEvent) override;
+#endif
+	//~ End UObject interface
+
+private:
+	/** Size of the input texture (Material and UMG only) */
+	UPROPERTY(EditAnywhere, Category = "Render Settings", Meta = (AllowPrivateAccess = true))
+	FVector2D InputSize{ 1024.f, 1024.f };
+
+	/** Number of times the pixelmapping input is downsampled */
+	UPROPERTY(EditAnywhere, Category = "Filtering", Meta = (ClampMin = "0", ClampMax = "256", UIMin = "0", UIMax = "16", AllowPrivateAccess = true))
+	int32 NumDownSamplePasses = 0;
+
+	/** Defines how the texture is resized after filtering */
+	UPROPERTY(EditAnywhere, Category = "Filtering", Meta = (AllowPrivateAccess = true))
+	EDMXPixelMappingRenderingPreprocessorSizeMode OutputSizeMode = EDMXPixelMappingRenderingPreprocessorSizeMode::SameAsInput;
+
+	/** Size of the rendered texture */
+	UPROPERTY(EditAnywhere, Category = "Filtering", Meta = (DisplayName = "Custom Size", ClampMin = "1", ClampMax = 2048, UIMin = 1, UIMax = 2048, EditCondition = "OutputSizeMode == EDMXPixelMappingRenderingPreprocessorSizeMode::Custom", EditConditionHides, AllowPrivateAccess = true))
+	FIntPoint CustomOutputSize { 256, 256 };
+	
+	/** Filter material applied to the rendered input */
+	UPROPERTY(EditAnywhere, BlueprintReadOnly, Category = "Filtering", Meta = (AllowPrivateAccess = true))
+	TObjectPtr<UMaterialInterface> FilterMaterial;
+
+	/** Actual material instance dynamic applied to the rendered input */
+	UPROPERTY(Transient)
+	TObjectPtr<UMaterialInstanceDynamic> FilterMID;
+
+	/** If true, the filter material is applied each downsample pass, otherwise only once after the last pass */
+	UPROPERTY(EditAnywhere, Category = "Filtering", Meta = (AllowPrivateAccess = true, DisplayName = "Apply Material each pass"))
+	bool bApplyFilterMaterialEachDownsamplePass = true;
+
+	/** Blur distance applied, only applicable if the filter matierial has a "BlurDistance" parameter */
+	UPROPERTY(EditAnywhere, Category = "Filtering", Meta = (UIMin = "0", UIMax = "1", SliderExponent = 2.0, AllowPrivateAccess = true))
+	float BlurDistance = 0.02f;
+
+	/** Parameter name of the filter material to receive the input texture */
+	UPROPERTY(EditAnywhere, AdvancedDisplay, Category = "Filtering", Meta = (AllowPrivateAccess = true))
+	FName InputTextureParameterName = "InputTexture";
+
+	/** The blur distance parameter name */
+	UPROPERTY(EditAnywhere, AdvancedDisplay, Category = "Filtering", Meta = (AllowPrivateAccess = true))
+	FName BlurDistanceParamterName = "BlurDistance";
+
+	/** Proxy to render the input */
+	TSharedPtr<UE::DMXPixelMapping::Rendering::Preprocess::Private::IPreprocessRenderInputProxy> RenderInputProxy;
+
+	/** Proxy to render the input */
+	TSharedPtr<UE::DMXPixelMapping::Rendering::Preprocess::Private::IPreprocessApplyFilterMaterialProxy> ApplyMaterialProxy;
+};
diff --git a/Engine/Plugins/VirtualProduction/DMX/DMXPixelMapping/Source/DMXPixelMappingRenderer/Public/IDMXPixelMappingRenderer.h b/Engine/Plugins/VirtualProduction/DMX/DMXPixelMapping/Source/DMXPixelMappingRenderer/Public/IDMXPixelMappingRenderer.h
--- a/Engine/Plugins/VirtualProduction/DMX/DMXPixelMapping/Source/DMXPixelMappingRenderer/Public/IDMXPixelMappingRenderer.h
+++ b/Engine/Plugins/VirtualProduction/DMX/DMXPixelMapping/Source/DMXPixelMappingRenderer/Public/IDMXPixelMappingRenderer.h
@@ -30,38 +30,6 @@ enum class EDMXPixelShaderBlendingQuality : uint8
 	MAX
 };
 
-namespace UE::DMXPixelMapping
-{
-	/** Parameters for the Input Texture Renderer */
-	struct FDMXPixelMappingInputTextureRenderingParameters
-	{
-		/** Number of times a texture is downsampled. E.g. when texture size is 512px and is downsampled 3 times, its resulting size is 64px */
-		int32 NumDownsamplePasses = 0;
-
-		/** The post process material. If null, no post process material is applied */
-		UMaterialInstanceDynamic* PostProcessMID = nullptr;
-
-		/** The input texture parameter name of the post process material */
-		FName PostProcessMaterialInputTextureParameterName;
-
-		/** The input texture parameter name of the post process material */
-		FName BlurDistanceParameterName;
-
-		/** The blur distance of the post process material */
-		float BlurDistance = .2f;
-
-		/**
-		 * If true, applies post process material each downsample pass.
-		 * If false applies the post process material once after the last downsample pass, or direct if the input is not downsampled.
-		 * Only applicable if a post process material is set.
-		 */
-		bool bApplyPostProcessMaterialEachDownsamplePass = true;
-
-		/** Size of the rendered texture */
-		FVector2D OutputSize{ 1.f, 1.f };
-	};
-
-}
 
 /**
  * Downsample pixel preview rendering params.
@@ -145,8 +113,6 @@ struct FDMXPixelMappingRenderTextureParams
 };
 
 
-
-
 /**
  * The public interface of the Pixel Mapping renderer instance interface.
  */
@@ -159,17 +125,6 @@ public:
 public:
 	/** Virtual destructor */
 	virtual ~IDMXPixelMappingRenderer() = default;
-		
-	/**
-	 * Blurs input texture onto desination texture
-	 *
-	 * @param InputTexture						The input texture that is being processed
-	 * @param Params							Parameters for post processing.
-	 */
-	virtual void PostProcessTexture(UTexture* InputTexture, const UE::DMXPixelMapping::FDMXPixelMappingInputTextureRenderingParameters& Params) const = 0;
-
-	/** Gets the post processed texture. May return nullptr while the texture is not rendered yet. */
-	virtual UTexture* GetPostProcessedTexture() const = 0;
 
 	/**
 	 * Pixelmapping specific, downsample and draw input texture to destination texture.
@@ -193,6 +148,7 @@ public:
 	 * @param InRenderTarget				2D render target texture resource
 	 * @param InMaterialInterface			Material to use
 	 */
+	 UE_DEPRECATED(5.3, "Please use FWidgetRenderer::DrawWidget or similar engine methods instead.")
 	virtual void RenderMaterial(UTextureRenderTarget2D* InRenderTarget, UMaterialInterface* InMaterialInterface) const = 0;
 
 	/**
@@ -201,6 +157,7 @@ public:
 	 * @param InRenderTarget				2D render target texture resource
 	 * @param InUserWidget					UMG widget to use
 	 */
+	 UE_DEPRECATED(5.3, "Please use FWidgetRenderer::DrawWidget or similar engine methods instead.")
 	virtual void RenderWidget(UTextureRenderTarget2D* InRenderTarget, UUserWidget* InUserWidget) const  = 0;
 
 	/**
diff --git a/Engine/Plugins/VirtualProduction/DMX/DMXPixelMapping/Source/DMXPixelMappingRuntime/Private/Components/DMXPixelMappingRendererComponent.cpp b/Engine/Plugins/VirtualProduction/DMX/DMXPixelMapping/Source/DMXPixelMappingRuntime/Private/Components/DMXPixelMappingRendererComponent.cpp
--- a/Engine/Plugins/VirtualProduction/DMX/DMXPixelMapping/Source/DMXPixelMappingRuntime/Private/Components/DMXPixelMappingRendererComponent.cpp
+++ b/Engine/Plugins/VirtualProduction/DMX/DMXPixelMapping/Source/DMXPixelMappingRuntime/Private/Components/DMXPixelMappingRendererComponent.cpp
@@ -2,37 +2,36 @@
 
 #include "Components/DMXPixelMappingRendererComponent.h"
 
-#include "DMXPixelMappingMainStreamObjectVersion.h"
-#include "DMXPixelMappingTypes.h"
-#include "DMXStats.h"
-#include "IDMXPixelMappingRenderer.h"
-#include "IDMXPixelMappingRendererModule.h"
+#include "Async/Async.h"
+#include "Blueprint/UserWidget.h"
 #include "Components/DMXPixelMappingFixtureGroupComponent.h"
 #include "Components/DMXPixelMappingFixtureGroupItemComponent.h"
 #include "Components/DMXPixelMappingMatrixComponent.h"
 #include "Components/DMXPixelMappingRootComponent.h"
 #include "Components/DMXPixelMappingScreenComponent.h"
-#include "Rendering/RenderInputTextureProxy.h"
-
-#include "Async/Async.h"
-#include "Blueprint/UserWidget.h"
+#include "DMXPixelMappingPreprocessRenderer.h"
+#include "DMXPixelMappingMainStreamObjectVersion.h"
+#include "DMXPixelMappingTypes.h"
+#include "DMXStats.h"
 #include "Engine/TextureRenderTarget2D.h"
-#include "Materials/MaterialInstanceDynamic.h"
+#include "IDMXPixelMappingRenderer.h"
+#include "IDMXPixelMappingRendererModule.h"
 #include "Modulators/DMXModulator.h"
+#include "RenderingThread.h"
 #include "TextureResource.h"
 #include "UObject/ConstructorHelpers.h"
 #include "UObject/Package.h"
 #include "Widgets/Layout/SConstraintCanvas.h"
 
-
 #if WITH_EDITOR
-
 #include "LevelEditor.h"
 #endif
 
 
 DECLARE_CYCLE_STAT(TEXT("PixelMapping Render"), STAT_DMXPixelMappingRender, STATGROUP_DMX);
-DECLARE_CYCLE_STAT(TEXT("PixelMapping Send DMX Worker"), STAT_DMXPixelMappingSendDMX, STATGROUP_DMX);
+DECLARE_CYCLE_STAT(TEXT("PixelMapping SendDMX"), STAT_DMXPixelMappingSendDMX, STATGROUP_DMX);
+DECLARE_CYCLE_STAT(TEXT("PixelMapping RenderInputTexture"), STAT_RenderInputTexture, STATGROUP_DMX);
+DECLARE_CYCLE_STAT(TEXT("PixelMapping GetTotalDownsamplePixelCount"), STAT_GetTotalDownsamplePixelCount, STATGROUP_DMX);
 
 
 const FIntPoint UDMXPixelMappingRendererComponent::MaxDownsampleBufferTargetSize = FIntPoint(4096);
@@ -41,7 +40,7 @@ const FLinearColor UDMXPixelMappingRendererComponent::ClearTextureColor = FLinea
 UDMXPixelMappingRendererComponent::UDMXPixelMappingRendererComponent()
 {
 	SetSize(FVector2D(100.f, 100.f));
-
+	
 #if WITH_EDITOR
 	ConstructorHelpers::FObjectFinder<UTexture> DefaultTexture(TEXT("Texture2D'/Engine/VREditor/Devices/Vive/UE4_Logo.UE4_Logo'"), LOAD_NoWarn);
 	if (ensureAlwaysMsgf(DefaultTexture.Succeeded(), TEXT("Failed to load Texture2D'/Engine/VREditor/Devices/Vive/UE4_Logo.UE4_Logo'")))
@@ -56,12 +55,8 @@ UDMXPixelMappingRendererComponent::UDMXPixelMappingRendererComponent()
 	}
 #endif
 	
+	PreprocessRenderer = CreateDefaultSubobject<UDMXPixelMappingPreprocessRenderer>("PreprocessRenderer");
 	Brightness = 1.0f;
-
-#if WITH_EDITOR
-	// Default to lock in designer, since for new renderers, the texture is the default
-	bLockInDesigner = true;
-#endif
 }
 
 const FName& UDMXPixelMappingRendererComponent::GetNamePrefix()
@@ -70,6 +65,11 @@ const FName& UDMXPixelMappingRendererComponent::GetNamePrefix()
 	return NamePrefix;
 }
 
+bool UDMXPixelMappingRendererComponent::CanBeMovedTo(const UDMXPixelMappingBaseComponent* Component) const
+{
+	return Component && Component->IsA<UDMXPixelMappingRootComponent>();
+}
+
 void UDMXPixelMappingRendererComponent::Serialize(FArchive& Ar)
 {
 	Super::Serialize(Ar);
@@ -81,15 +81,11 @@ void UDMXPixelMappingRendererComponent::Serialize(FArchive& Ar)
 		{
 			if (RendererType == EDMXPixelMappingRendererType::Texture)
 			{
-#if WITH_EDITOR
-				bLockInDesigner = true;
-#endif
 				// Refresh the size of the texture if that is used as input
 				if (InputTexture)
 				{
 					if (const FTextureResource* TextureResource = InputTexture->GetResource())
 					{
-						// Set to the texture size
 						const FVector2D NewSize = FVector2D(TextureResource->GetSizeX(), TextureResource->GetSizeY());
 						SetSize(NewSize);
 					}
@@ -126,51 +122,7 @@ void UDMXPixelMappingRendererComponent::PostEditChangeChainProperty(FPropertyCha
 	Super::PostEditChangeChainProperty(PropertyChangedChainEvent);
 
 	const FName PropertyName = PropertyChangedChainEvent.GetPropertyName();
-	if (PropertyChangedChainEvent.GetPropertyName() == UDMXPixelMappingOutputComponent::GetSizeXPropertyName() ||
-		PropertyChangedChainEvent.GetPropertyName() == UDMXPixelMappingOutputComponent::GetSizeYPropertyName())
-	{
-		// The target always needs be within GMaxTextureDimensions, larger dimensions are not supported by the engine
-		const uint32 MaxTextureDimensions = GetMax2DTextureDimension();
-
-		if (GetSize().X > MaxTextureDimensions ||
-			GetSize().Y > MaxTextureDimensions)
-		{
-			const float NewSizeX = FMath::Clamp(GetSize().X, 0.0f, static_cast<float>(MaxTextureDimensions));
-			const float NewSizeY = FMath::Clamp(GetSize().Y, 0.0f, static_cast<float>(MaxTextureDimensions));
-			const FVector2D NewSize(NewSizeX, NewSizeY);
-			SetSize(NewSize);
-
-			UE_LOG(LogDMXPixelMappingRuntime, Warning, TEXT("Pixel mapping textures are limited to engine's max texture dimension %dx%d"), MaxTextureDimensions, MaxTextureDimensions);
-		}
-	} 
-	else if (PropertyChangedChainEvent.GetPropertyName() == GET_MEMBER_NAME_CHECKED(UDMXPixelMappingRendererComponent, RendererType))
-	{
-		if (RendererType == EDMXPixelMappingRendererType::Texture)
-		{
-			// Prevent the size from being edited via its edit condition
-			bLockInDesigner = true;
-		}
-		else
-		{
-			bLockInDesigner = false;
-		}
-	}
-	else if (PropertyChangedChainEvent.GetPropertyName() == GET_MEMBER_NAME_CHECKED(UDMXPixelMappingRendererComponent, InputWidget))
-	{
-		if (InputWidget && UserWidget && InputWidget->GetClass() != UserWidget->GetClass())
-		{
-			// UMG just tries to expand to the max possible size. Instead of using that we set a smaller, reasonable size here. 
-			// This doesn't offer a solution to the adaptive nature of UMG, but implies to the user how to deal with the issue.
-			constexpr float DefaultUMGSizeX = 1024.f;
-			constexpr float DefaultUMGSizeY = 768.f;
-
-			SetSize(FVector2D(DefaultUMGSizeX, DefaultUMGSizeY));
-			ResizePreviewRenderTarget(DefaultUMGSizeX, DefaultUMGSizeY);
-		}
-
-		UpdateInputWidget(InputWidget);
-	}
-	else if (PropertyChangedChainEvent.GetPropertyName() == GET_MEMBER_NAME_CHECKED(UDMXPixelMappingRendererComponent, Brightness))
+	if (PropertyName == GET_MEMBER_NAME_CHECKED(UDMXPixelMappingRendererComponent, Brightness))
 	{
 		const TSharedPtr<IDMXPixelMappingRenderer>& Renderer = GetRenderer();
 		if (Renderer.IsValid())
@@ -415,20 +367,9 @@ UWorld* UDMXPixelMappingRendererComponent::GetWorld() const
 	return World;
 }
 
-void UDMXPixelMappingRendererComponent::UpdateInputWidget(TSubclassOf<UUserWidget> InInputWidget)
-{
-	UserWidget = CreateWidget(GetWorld(), InInputWidget);
-}
-
 void UDMXPixelMappingRendererComponent::Initialize()
 {
-	if (InputTexture)
-	{
-		InputTexture->SetForceMipLevelsToBeResident(30.0f);
-		InputTexture->WaitForStreaming();
-	}
-
-	if (UserWidget == nullptr && InputWidget != nullptr)
+	if (InputWidget)
 	{
 		UserWidget = CreateWidget(GetWorld(), InputWidget);
 	}
@@ -439,31 +380,25 @@ void UDMXPixelMappingRendererComponent::Initialize()
 		PixelMappingRenderer->SetBrightness(Brightness);
 	}
 
-#if WITH_EDITOR
-	// Before 4.27 the 'bLockInDesigner' edit condition for size did not exist, apply it here where needed
-	if (RendererType == EDMXPixelMappingRendererType::Texture)
+	switch (RendererType)
 	{
-		// Prevent the size from being edited via its edit condition
-		bLockInDesigner = true;
+	case(EDMXPixelMappingRendererType::Texture):
+		PreprocessRenderer->SetInputTexture(InputTexture.Get());
+		break;
 
-		// Refresh the size of the texture in case it changed externally
-		if (InputTexture)
-		{
-			if (const FTextureResource* TextureResource = InputTexture->GetResource())
-			{
-				// Set to the texture size
-				const FVector2D NewSize = FVector2D(TextureResource->GetSizeX(), TextureResource->GetSizeY());
-				SetSize(NewSize);
-			}
-		}
-	}
-	else
-	{
-		bLockInDesigner = false;
+	case(EDMXPixelMappingRendererType::Material):
+		PreprocessRenderer->SetInputMaterial(InputMaterial.Get());
+		break;
+
+	case(EDMXPixelMappingRendererType::UMG):
+		PreprocessRenderer->SetInputUserWidget(UserWidget.Get());
+		break;
+
+	default:
+		checkf(0, TEXT("Invalid Renderer Type in DMXPixelMappingRendererComponent"));
 	}
-#endif
 
-	UpdateInputTextureRenderProxy();
+	SetSize(PreprocessRenderer->GetResultingSize2D());
 }
 
 UTextureRenderTarget2D* UDMXPixelMappingRendererComponent::CreateRenderTarget(const FName& InBaseName)
@@ -479,34 +414,28 @@ UTextureRenderTarget2D* UDMXPixelMappingRendererComponent::CreateRenderTarget(co
 
 void UDMXPixelMappingRendererComponent::RendererInputTexture()
 {
-	if (RenderInputTextureProxy.IsValid())
-	{
-		RenderInputTextureProxy->Render();
-	}
+	SCOPE_CYCLE_COUNTER(STAT_RenderInputTexture);
 
-#if WITH_EDITOR
-	if (RendererType == EDMXPixelMappingRendererType::Texture)
-	{
-		if (InputTexture != nullptr && InputTexture->GetResource())
-		{
-			ResizePreviewRenderTarget(InputTexture->GetResource()->GetSizeX(), InputTexture->GetResource()->GetSizeY());
-		}
-	}
-	else
-	{
-		ResizePreviewRenderTarget(GetSize().X, GetSize().Y);
-	}
-#endif
-}
+	PreprocessRenderer->Render();
 
-UTexture* UDMXPixelMappingRendererComponent::GetRendererInputTexture() const
-{
-	return RenderInputTextureProxy.IsValid() ? RenderInputTextureProxy->GetRenderedTexture() : nullptr;
+//#if WITH_EDITOR
+//	if (RendererType == EDMXPixelMappingRendererType::Texture)
+//	{
+//		if (InputTexture && InputTexture->GetResource())
+//		{
+//			ResizePreviewRenderTarget(InputTexture->GetResource()->GetSizeX(), InputTexture->GetResource()->GetSizeY());
+//		}
+//	}
+//	else
+//	{
+//		ResizePreviewRenderTarget(GetSize().X, GetSize().Y);
+//	}
+//#endif
 }
 
-bool UDMXPixelMappingRendererComponent::CanBeMovedTo(const UDMXPixelMappingBaseComponent* Component) const
+UTexture* UDMXPixelMappingRendererComponent::GetRendererInputTexture() const
 {
-	return Component && Component->IsA<UDMXPixelMappingRootComponent>();
+	return PreprocessRenderer->GetRenderedTexture();
 }
 
 void UDMXPixelMappingRendererComponent::CreateOrUpdateDownsampleBufferTarget()
@@ -671,6 +600,8 @@ bool UDMXPixelMappingRendererComponent::IsPixelRangeValid(const int32 InDownsamp
 
 int32 UDMXPixelMappingRendererComponent::GetTotalDownsamplePixelCount()
 {
+	SCOPE_CYCLE_COUNTER(STAT_GetTotalDownsamplePixelCount);
+
 	FScopeLock ScopeLock(&DownsampleBufferCS);
 
 	// Reset pixel counter
@@ -694,31 +625,3 @@ int32 UDMXPixelMappingRendererComponent::GetTotalDownsamplePixelCount()
 
 	return DownsamplePixelCount;
 }
-
-void UDMXPixelMappingRendererComponent::UpdateInputTextureRenderProxy()
-{
-	using namespace UE::DMXPixelMapping::Rendering::Private;
-
-	// Update params
-	RenderInputTextureParams.OutputSize = GetSize();
-	RenderInputTextureParams.PostProcessMID = PostProcessMaterial ? UMaterialInstanceDynamic::Create(PostProcessMaterial, GetTransientPackage()) : nullptr;
-
-	RenderInputTextureProxy.Reset();
-	switch (RendererType)
-	{
-	case(EDMXPixelMappingRendererType::Texture):
-		RenderInputTextureProxy = MakeShared<FRenderInputTextureProxy<UTexture>>(PixelMappingRenderer.ToSharedRef(), InputTexture.Get(), RenderInputTextureParams);
-		break;
-
-	case(EDMXPixelMappingRendererType::Material):
-		RenderInputTextureProxy = MakeShared<FRenderInputTextureProxy<UMaterialInterface>>(PixelMappingRenderer.ToSharedRef(), InputMaterial.Get(), RenderInputTextureParams, GetSize());
-		break;
-
-	case(EDMXPixelMappingRendererType::UMG):
-		RenderInputTextureProxy = MakeShared<FRenderInputTextureProxy<UUserWidget>>(PixelMappingRenderer.ToSharedRef(), InputWidget.GetDefaultObject(), RenderInputTextureParams, GetSize());
-		break;
-
-	default:
-		checkf(0, TEXT("Invalid Renderer Type in DMXPixelMappingRendererComponent"));
-	}
-}
diff --git a/Engine/Plugins/VirtualProduction/DMX/DMXPixelMapping/Source/DMXPixelMappingRuntime/Private/Rendering/RenderInputTextureProxy.cpp b/Engine/Plugins/VirtualProduction/DMX/DMXPixelMapping/Source/DMXPixelMappingRuntime/Private/Rendering/RenderInputTextureProxy.cpp
deleted file mode 100644
diff --git a/Engine/Plugins/VirtualProduction/DMX/DMXPixelMapping/Source/DMXPixelMappingRuntime/Private/Rendering/RenderInputTextureProxy.h b/Engine/Plugins/VirtualProduction/DMX/DMXPixelMapping/Source/DMXPixelMappingRuntime/Private/Rendering/RenderInputTextureProxy.h
deleted file mode 100644
diff --git a/Engine/Plugins/VirtualProduction/DMX/DMXPixelMapping/Source/DMXPixelMappingRuntime/Public/Components/DMXPixelMappingOutputComponent.h b/Engine/Plugins/VirtualProduction/DMX/DMXPixelMapping/Source/DMXPixelMappingRuntime/Public/Components/DMXPixelMappingOutputComponent.h
--- a/Engine/Plugins/VirtualProduction/DMX/DMXPixelMapping/Source/DMXPixelMappingRuntime/Public/Components/DMXPixelMappingOutputComponent.h
+++ b/Engine/Plugins/VirtualProduction/DMX/DMXPixelMapping/Source/DMXPixelMappingRuntime/Public/Components/DMXPixelMappingOutputComponent.h
@@ -122,7 +122,7 @@ public:
 	/** Sets the size */
 	virtual void SetSize(const FVector2D& NewSize);
 
-	/** Get the size. Note this is the position before a layout has been applied, use GetLayoutedSize() to get the final position.. */
+	/** Get the size */
 	FVector2D GetSize() const { return FVector2D(SizeX, SizeY); }
 
 	/** Helper that returns render component if available */
diff --git a/Engine/Plugins/VirtualProduction/DMX/DMXPixelMapping/Source/DMXPixelMappingRuntime/Public/Components/DMXPixelMappingRendererComponent.h b/Engine/Plugins/VirtualProduction/DMX/DMXPixelMapping/Source/DMXPixelMappingRuntime/Public/Components/DMXPixelMappingRendererComponent.h
--- a/Engine/Plugins/VirtualProduction/DMX/DMXPixelMapping/Source/DMXPixelMappingRuntime/Public/Components/DMXPixelMappingRendererComponent.h
+++ b/Engine/Plugins/VirtualProduction/DMX/DMXPixelMapping/Source/DMXPixelMappingRuntime/Public/Components/DMXPixelMappingRendererComponent.h
@@ -2,19 +2,18 @@
 
 #pragma once
 
-#include "IDMXPixelMappingRenderer.h"
+#include "DMXPixelMappingPreprocessRenderer.h"
 #include "Components/DMXPixelMappingOutputComponent.h"
+#include "IDMXPixelMappingRenderer.h"
 #include "Library/DMXEntityReference.h"
-#include "Rendering/IDMXPixelMappingRenderInputTextureProxy.h"
-
 #include "Templates/SubclassOf.h"
 
 #include "DMXPixelMappingRendererComponent.generated.h"
 
 enum class EDMXPixelMappingRendererType : uint8;
 class UDMXPixelMappingLayoutScript;
+class UDMXPixelMappingPreprocessRenderer;
 
-enum class EMapChangeType : uint8;
 class UMaterialInterface;
 class UTexture;
 class UUserWidget;
@@ -46,6 +45,7 @@ public:
 
 	//~ Begin UDMXPixelMappingBaseComponent implementation
 	virtual const FName& GetNamePrefix() override;
+	virtual bool CanBeMovedTo(const UDMXPixelMappingBaseComponent* Component) const override;
 	virtual void ResetDMX() override;
 	virtual void SendDMX() override;
 	virtual void Render() final;
@@ -130,9 +130,6 @@ public:
 	void EmptyDownsampleBuffer();
 
 private:
-	/** Generate new input widget based on UMG */
-	void UpdateInputWidget(TSubclassOf<UUserWidget> InInputWidget);
-
 	/** Resize output texture for editor preview */
 	void ResizePreviewRenderTarget(uint32 InSizeX, uint32 InSizeY);
 
@@ -150,42 +147,38 @@ public:
 	UFUNCTION(BlueprintCallable, Category = "DMX")
 	bool GetPixelMappingComponentModulators(FDMXEntityFixturePatchRef FixturePatchRef, TArray<UDMXModulator*>& DMXModulators);
 
+	/** Returns the preprocess renderer */
+	UDMXPixelMappingPreprocessRenderer* GetPreprocessRenderer() { return PreprocessRenderer; }
+
+#if WITH_EDITOR
+	/** Returns the component canvas used for this widget */
+	FORCEINLINE TSharedPtr<SConstraintCanvas> GetComponentsCanvas() const { return ComponentsCanvas; }
+#endif // WITH_EDITOR
+
 	/** Type of rendering, Texture, Material, UMG, etc... */
 	UPROPERTY(EditAnywhere, BlueprintReadOnly, Category = "Render Settings")
 	EDMXPixelMappingRendererType RendererType;
 
-	/** Texture to Downsampling */
+	/** The texture used for pixel mapping */
 	UPROPERTY(EditAnywhere, BlueprintReadOnly, Category = "Render Settings")
 	TObjectPtr<UTexture> InputTexture;
 
-	/** Material to Downsampling */
-	UPROPERTY(EditAnywhere, BlueprintReadOnly, Category = "Render Settings", meta = (DisplayName = "User Interface Material"))
+	/** The material used for pixel mapping */
+	UPROPERTY(EditAnywhere, BlueprintReadOnly, Category = "Render Settings", Meta = (DisplayName = "User Interface Material"))
 	TObjectPtr<UMaterialInterface> InputMaterial;
 
-	/** UMG to Downsampling */
+	/** The UMG widget used for pixel mapping */
 	UPROPERTY(EditAnywhere, BlueprintReadOnly, Category = "Render Settings")
 	TSubclassOf<UUserWidget> InputWidget;
 
 	/** The brightness of the renderer */
-	UPROPERTY(EditAnywhere, BlueprintReadOnly, Category = "Render Settings", meta = (ClampMin = "0", ClampMax = "1", UIMin = "0", UIMax = "1"))
+	UPROPERTY(EditAnywhere, BlueprintReadOnly, Category = "Render Settings", Meta = (ClampMin = "0", ClampMax = "1", UIMin = "0", UIMax = "1"))
 	float Brightness = 1.f;
 
-	/** Post process material applied to the rendered input */
-	UPROPERTY(EditAnywhere, BlueprintReadOnly, Category = "Post Process")
-	TObjectPtr<UMaterialInterface> PostProcessMaterial;
-
 	/** Layout script for the children of this component (hidden in customizations and displayed in its own panel). */
 	UPROPERTY(EditAnywhere, Instanced, Category = "Layout")
 	TObjectPtr<UDMXPixelMappingLayoutScript> LayoutScript;
 
-	/** Check if a Component can be moved under another one (used for copy/move/duplicate) */
-	virtual bool CanBeMovedTo(const UDMXPixelMappingBaseComponent* Component) const override;
-
-#if WITH_EDITOR
-	/** Returns the component canvas used for this widget */
-	FORCEINLINE TSharedPtr<SConstraintCanvas> GetComponentsCanvas() const { return ComponentsCanvas; }
-#endif // WITH_EDITOR
-
 private:
 	/** Retrieve total count of all output targets that support shared rendering and updates a counter. O(n) */
 	int32 GetTotalDownsamplePixelCount();
@@ -193,15 +186,8 @@ private:
 	/** Helper function checks the downsample pixel range */
 	bool IsPixelRangeValid(const int32 InDownsamplePixelIndexStart, const int32 InDownsamplePixelIndexEnd) const;
 
-	/** Updates the Input Texture Render Proxy, depending on the selected source */
-	void UpdateInputTextureRenderProxy();
-
-	/** Proxy responsible to render the input texture */
-	TSharedPtr<UE::DMXPixelMapping::Rendering::Private::IDMXPixelMappingRenderInputTextureProxy> RenderInputTextureProxy;
-
-	/** Parameters by which the input texture is rendered */
-	UPROPERTY(EditAnywhere, Category = "Post Process", Meta = (ShowOnlyInnerProperties, AllowPrivateAccess = true))
-	FDMXPixelMappingRenderInputTextureParameters RenderInputTextureParams;
+	UPROPERTY(VisibleAnywhere, Instanced, Category = "Filtering")
+	TObjectPtr<UDMXPixelMappingPreprocessRenderer> PreprocessRenderer;
 
 	/** Reference to renderer */
 	TSharedPtr<IDMXPixelMappingRenderer> PixelMappingRenderer;
diff --git a/Engine/Plugins/VirtualProduction/DMX/DMXPixelMapping/Source/DMXPixelMappingRuntime/Public/Rendering/IDMXPixelMappingRenderInputTextureProxy.h b/Engine/Plugins/VirtualProduction/DMX/DMXPixelMapping/Source/DMXPixelMappingRuntime/Public/Rendering/IDMXPixelMappingRenderInputTextureProxy.h
deleted file mode 100644
