DMX - Add an option to specify the range of the Pixel Mapping xyY color space

diff --git a/Engine/Plugins/VirtualProduction/DMX/DMXPixelMapping/Source/DMXPixelMappingEditor/Private/Customizations/DMXPixelMappingColorSpaceDetails_xyY.cpp b/Engine/Plugins/VirtualProduction/DMX/DMXPixelMapping/Source/DMXPixelMappingEditor/Private/Customizations/DMXPixelMappingColorSpaceDetails_xyY.cpp
new file mode 100644
--- /dev/null
+++ b/Engine/Plugins/VirtualProduction/DMX/DMXPixelMapping/Source/DMXPixelMappingEditor/Private/Customizations/DMXPixelMappingColorSpaceDetails_xyY.cpp
@@ -0,0 +1,61 @@
+// Copyright Epic Games, Inc. All Rights Reserved.
+
+#include "Customizations/DMXPixelMappingColorSpaceDetails_xyY.h"
+
+#include "ColorSpace/DMXPixelMappingColorSpace_xyY.h"
+#include "DetailCategoryBuilder.h"
+#include "DetailLayoutBuilder.h"
+#include "DetailWidgetRow.h"
+#include "Styling/AppStyle.h"
+#include "Widgets/SBoxPanel.h"
+#include "Widgets/Text/STextBlock.h"
+
+
+#define LOCTEXT_NAMESPACE "DMXPixelMappingColorSpaceDetails_xyY"
+
+TSharedRef<IDetailCustomization> FDMXPixelMappingColorSpaceDetails_xyY::MakeInstance()
+{
+	return MakeShared<FDMXPixelMappingColorSpaceDetails_xyY>();
+}
+
+void FDMXPixelMappingColorSpaceDetails_xyY::CustomizeDetails(IDetailLayoutBuilder& DetailBuilder)
+{
+	IDetailCategoryBuilder& CategoryBuilder = DetailBuilder.EditCategory("XY");
+
+	// Add attributes before the range
+	CategoryBuilder.AddProperty(DetailBuilder.GetProperty(GET_MEMBER_NAME_CHECKED(UDMXPixelMappingColorSpace_xyY, XAttribute)));
+	CategoryBuilder.AddProperty(DetailBuilder.GetProperty(GET_MEMBER_NAME_CHECKED(UDMXPixelMappingColorSpace_xyY, YAttribute)));
+
+	// Present the range percents property with a percent sign 
+	const TSharedRef<IPropertyHandle> ColorSpaceRangePercentsPropertyHandle = DetailBuilder.GetProperty(GET_MEMBER_NAME_CHECKED(UDMXPixelMappingColorSpace_xyY, ColorSpaceRangePercents));
+
+	DetailBuilder.HideProperty(ColorSpaceRangePercentsPropertyHandle);
+	
+	CategoryBuilder
+		.AddCustomRow(LOCTEXT("RangePercentsFilter", "Range"))
+		.NameContent()
+		[
+			ColorSpaceRangePercentsPropertyHandle->CreatePropertyNameWidget()
+		]
+		.ValueContent()
+		[
+			SNew(SHorizontalBox)
+
+			+ SHorizontalBox::Slot()
+			.AutoWidth()
+			[
+				ColorSpaceRangePercentsPropertyHandle->CreatePropertyValueWidget()
+			]
+			
+			+ SHorizontalBox::Slot()
+			.VAlign(VAlign_Center)
+			.Padding(2.f)
+			[
+				SNew(STextBlock)
+				.Font(FAppStyle::Get().GetFontStyle("PropertyWindow.NormalFont"))
+				.Text(LOCTEXT("PercentSign", "%"))
+			]
+		];
+}
+
+#undef LOCTEXT_NAMESPACE
diff --git a/Engine/Plugins/VirtualProduction/DMX/DMXPixelMapping/Source/DMXPixelMappingEditor/Private/Customizations/DMXPixelMappingColorSpaceDetails_xyY.h b/Engine/Plugins/VirtualProduction/DMX/DMXPixelMapping/Source/DMXPixelMappingEditor/Private/Customizations/DMXPixelMappingColorSpaceDetails_xyY.h
new file mode 100644
--- /dev/null
+++ b/Engine/Plugins/VirtualProduction/DMX/DMXPixelMapping/Source/DMXPixelMappingEditor/Private/Customizations/DMXPixelMappingColorSpaceDetails_xyY.h
@@ -0,0 +1,19 @@
+// Copyright Epic Games, Inc. All Rights Reserved.
+
+#pragma once
+
+#include "IDetailCustomization.h"
+
+
+/** Details customization for the xyY Color Space */
+class FDMXPixelMappingColorSpaceDetails_xyY
+	: public IDetailCustomization
+{
+public:
+	/** Creates an instance of this customization */
+	static TSharedRef<IDetailCustomization> MakeInstance();
+
+	//~ Begin IDetailCustomization interface 
+	virtual void CustomizeDetails(IDetailLayoutBuilder& DetailBuilder) override;
+	//~ End IDetailCustomization interface 
+};
diff --git a/Engine/Plugins/VirtualProduction/DMX/DMXPixelMapping/Source/DMXPixelMappingEditor/Private/Views/SDMXPixelMappingDetailsView.cpp b/Engine/Plugins/VirtualProduction/DMX/DMXPixelMapping/Source/DMXPixelMappingEditor/Private/Views/SDMXPixelMappingDetailsView.cpp
--- a/Engine/Plugins/VirtualProduction/DMX/DMXPixelMapping/Source/DMXPixelMappingEditor/Private/Views/SDMXPixelMappingDetailsView.cpp
+++ b/Engine/Plugins/VirtualProduction/DMX/DMXPixelMapping/Source/DMXPixelMappingEditor/Private/Views/SDMXPixelMappingDetailsView.cpp
@@ -2,8 +2,8 @@
 
 #include "Views/SDMXPixelMappingDetailsView.h"
 
-#include "DMXPixelMappingComponentReference.h"
 #include "ColorSpace/DMXPixelMappingColorSpace_RGBCMY.h"
+#include "ColorSpace/DMXPixelMappingColorSpace_xyY.h"
 #include "Components/DMXPixelMappingRootComponent.h"
 #include "Components/DMXPixelMappingFixtureGroupComponent.h"
 #include "Components/DMXPixelMappingFixtureGroupItemComponent.h"
@@ -12,19 +12,20 @@
 #include "Components/DMXPixelMappingMatrixCellComponent.h"
 #include "Components/DMXPixelMappingRendererComponent.h"
 #include "Customizations/DMXPixelMappingColorSpaceDetails_RGBCMY.h"
+#include "Customizations/DMXPixelMappingColorSpaceDetails_xyY.h"
 #include "Customizations/DMXPixelMappingDetailCustomization_FixtureGroup.h"
 #include "Customizations/DMXPixelMappingDetailCustomization_FixtureGroupItem.h"
 #include "Customizations/DMXPixelMappingDetailCustomization_Screen.h"
 #include "Customizations/DMXPixelMappingDetailCustomization_Renderer.h"
 #include "Customizations/DMXPixelMappingDetailCustomization_Matrix.h"
 #include "Customizations/DMXPixelMappingDetailCustomization_MatrixCell.h"
-#include "Toolkits/DMXPixelMappingToolkit.h"
-
 #include "DetailsViewArgs.h"
+#include "DMXPixelMappingComponentReference.h"
 #include "IDetailsView.h"
 #include "PropertyEditorDelegates.h"
 #include "PropertyEditorModule.h"
 #include "Modules/ModuleManager.h"
+#include "Toolkits/DMXPixelMappingToolkit.h"
 
 
 void SDMXPixelMappingDetailsView::Construct(const FArguments& InArgs, const TSharedPtr<FDMXPixelMappingToolkit>& InToolkit)
@@ -71,6 +72,7 @@ SDMXPixelMappingDetailsView::~SDMXPixelMappingDetailsView()
 		PropertyView->UnregisterInstancedCustomPropertyLayout(UDMXPixelMappingMatrixCellComponent::StaticClass());
 		
 		PropertyView->UnregisterInstancedCustomPropertyLayout(UDMXPixelMappingColorSpace_RGBCMY::StaticClass());
+		PropertyView->UnregisterInstancedCustomPropertyLayout(UDMXPixelMappingColorSpace_xyY::StaticClass());
 	}
 }
 
@@ -136,4 +138,7 @@ void SDMXPixelMappingDetailsView::RegisterCustomizations()
 
 	FOnGetDetailCustomizationInstance ColorSpaceCustomizationInstance_RGBCMY = FOnGetDetailCustomizationInstance::CreateStatic(&FDMXPixelMappingColorSpaceDetails_RGBCMY::MakeInstance);
 	PropertyView->RegisterInstancedCustomPropertyLayout(UDMXPixelMappingColorSpace_RGBCMY::StaticClass(), ColorSpaceCustomizationInstance_RGBCMY);
+
+	FOnGetDetailCustomizationInstance ColorSpaceCustomizationInstance_xyY = FOnGetDetailCustomizationInstance::CreateStatic(&FDMXPixelMappingColorSpaceDetails_xyY::MakeInstance);
+	PropertyView->RegisterInstancedCustomPropertyLayout(UDMXPixelMappingColorSpace_xyY::StaticClass(), ColorSpaceCustomizationInstance_xyY);
 }
diff --git a/Engine/Plugins/VirtualProduction/DMX/DMXPixelMapping/Source/DMXPixelMappingRuntime/Private/ColorSpace/DMXPixelMappingColorSpace_xyY.cpp b/Engine/Plugins/VirtualProduction/DMX/DMXPixelMapping/Source/DMXPixelMappingRuntime/Private/ColorSpace/DMXPixelMappingColorSpace_xyY.cpp
--- a/Engine/Plugins/VirtualProduction/DMX/DMXPixelMapping/Source/DMXPixelMappingRuntime/Private/ColorSpace/DMXPixelMappingColorSpace_xyY.cpp
+++ b/Engine/Plugins/VirtualProduction/DMX/DMXPixelMapping/Source/DMXPixelMappingRuntime/Private/ColorSpace/DMXPixelMappingColorSpace_xyY.cpp
@@ -8,7 +8,9 @@ UDMXPixelMappingColorSpace_xyY::UDMXPixelMappingColorSpace_xyY()
 	, YAttribute("CIE_Y")
 	, LuminanceAttribute("Dimmer")
 	, SRGBColorSpace(UE::Color::EColorSpace::sRGB)
-{}
+{
+	ColorSpaceRangePercents = ColorSpaceRange * 100.f;
+}
 
 void UDMXPixelMappingColorSpace_xyY::SetRGBA(const FLinearColor& InColor)
 {
@@ -21,23 +23,37 @@ void UDMXPixelMappingColorSpace_xyY::SetRGBA(const FLinearColor& InColor)
 	const FVector4 XYZW = Matrix.TransformVector(FVector(InColor));
 	const FVector3d xyY = UE::Color::XYZToxyY(XYZW);
 
+	if (!ensureMsgf(ColorSpaceRange != 0.0, TEXT("Coversion in PixelMapping Color Space xyY failed. Color space range is 0.")))
+	{
+		return;
+	}
+
 	if (XAttribute.IsValid())
 	{
-		SetAttributeValue(XAttribute, xyY[0]);
+		const double x = xyY[0] / ColorSpaceRange;
+		SetAttributeValue(XAttribute, x);
 	}
 
 	if (YAttribute.IsValid())
 	{
-		SetAttributeValue(YAttribute, xyY[1]);
+		const double y = xyY[1] / ColorSpaceRange;
+		SetAttributeValue(YAttribute, y);
 	}
 
 	if (LuminanceAttribute.IsValid())
 	{
-		const float Luminance = FMath::Clamp(xyY[2], MinLuminance, MaxLuminance);
+		const double Luminance = FMath::Clamp(xyY[2], MinLuminance, MaxLuminance);
 		SetAttributeValue(LuminanceAttribute, Luminance);
 	}
 }
 
+void UDMXPixelMappingColorSpace_xyY::PostLoad()
+{
+	Super::PostLoad();
+
+	ColorSpaceRangePercents = ColorSpaceRange * 100.f;
+}
+
 #if WITH_EDITOR
 void UDMXPixelMappingColorSpace_xyY::PostEditChangeProperty(FPropertyChangedEvent& PropertyChangedEvent)
 {
@@ -67,5 +83,9 @@ void UDMXPixelMappingColorSpace_xyY::PostEditChangeProperty(FPropertyChangedEven
 			MinLuminance = MaxLuminance;
 		}
 	}
+	else if (PropertyName == GET_MEMBER_NAME_CHECKED(UDMXPixelMappingColorSpace_xyY, ColorSpaceRangePercents))
+	{
+		ColorSpaceRange = ColorSpaceRangePercents / 100.f;
+	}
 }
 #endif // WITH_EDITOR
diff --git a/Engine/Plugins/VirtualProduction/DMX/DMXPixelMapping/Source/DMXPixelMappingRuntime/Private/ColorSpace/DMXPixelMappingColorSpace_xyY.h b/Engine/Plugins/VirtualProduction/DMX/DMXPixelMapping/Source/DMXPixelMappingRuntime/Private/ColorSpace/DMXPixelMappingColorSpace_xyY.h
--- a/Engine/Plugins/VirtualProduction/DMX/DMXPixelMapping/Source/DMXPixelMappingRuntime/Private/ColorSpace/DMXPixelMappingColorSpace_xyY.h
+++ b/Engine/Plugins/VirtualProduction/DMX/DMXPixelMapping/Source/DMXPixelMappingRuntime/Private/ColorSpace/DMXPixelMappingColorSpace_xyY.h
@@ -32,25 +32,34 @@ public:
 	UPROPERTY(EditAnywhere, Category = "XY", Meta = (DisplayName = "y Attribute"))
 	FDMXAttributeName YAttribute;
 
+	/** Sets the range of the xyY color space. CIE 1931 is 100%. */
+	UPROPERTY(EditAnywhere, Transient, Category = "XY", Meta = (ClampMin = 0.1, ClampMax = 1000.0, UIMin = 50.0, UIMax = 100.0, SliderExponent = 3, DisplayName = "Color Space Range"))
+	float ColorSpaceRangePercents = 100.f;
+
 	/** Attribute sent for Y */
 	UPROPERTY(EditAnywhere, Category = "Luminance")
 	FDMXAttributeName LuminanceAttribute;
 
 	/** Min Luminance */
-	UPROPERTY(EditAnywhere, Category = "Luminance", Meta = (ClampMin = 0.0, ClampMax = 1.0, SliderMin = 0.0, SliderMax = 1.0, EditConditionHides, EditCondition = "LuminanceType == EDMXPixelMappingLuminanceType_RGBCMY::FromWhite || LuminanceType == EDMXPixelMappingLuminanceType_RGBCMY::FromAlpha"))
+	UPROPERTY(EditAnywhere, Category = "Luminance", Meta = (ClampMin = 0.0, ClampMax = 1.0, UIMin = 0.0, UIMax = 1.0, EditConditionHides, EditCondition = "LuminanceType == EDMXPixelMappingLuminanceType_RGBCMY::FromWhite || LuminanceType == EDMXPixelMappingLuminanceType_RGBCMY::FromAlpha"))
 	float MinLuminance = 0.f;
 
 	/** Max Luminance */
-	UPROPERTY(EditAnywhere, Category = "Luminance", Meta = (ClampMin = 0.0, ClampMax = 1.0, SliderMin = 0.0, SliderMax = 1.0, EditConditionHides, EditCondition = "LuminanceType == EDMXPixelMappingLuminanceType_RGBCMY::FromWhite || LuminanceType == EDMXPixelMappingLuminanceType_RGBCMY::FromAlpha"))
+	UPROPERTY(EditAnywhere, Category = "Luminance", Meta = (ClampMin = 0.0, ClampMax = 1.0, UIMin = 0.0, UIMax = 1.0, EditConditionHides, EditCondition = "LuminanceType == EDMXPixelMappingLuminanceType_RGBCMY::FromWhite || LuminanceType == EDMXPixelMappingLuminanceType_RGBCMY::FromAlpha"))
 	float MaxLuminance = 1.f;
 
 protected:
 	//~ Begin UObject Interface
+	virtual void PostLoad() override;
 #if WITH_EDITOR
 	virtual void PostEditChangeProperty(FPropertyChangedEvent& PropertyChangedEvent) override;
 #endif // WITH_EDITOR
 
 private:
+	/** The actual range property */
+	UPROPERTY()
+	float ColorSpaceRange = 1.0;
+
 	/** Cached sRGB color space, to avoid instantiating on each conversion */
 	UE::Color::FColorSpace SRGBColorSpace;
 };
