DMX - Fix CIS Static Analysis Win64

diff --git a/Engine/Plugins/VirtualProduction/DMX/DMXPixelMapping/Source/DMXPixelMappingRenderer/Private/PreprocessApplyFilterMaterialProxy.cpp b/Engine/Plugins/VirtualProduction/DMX/DMXPixelMapping/Source/DMXPixelMappingRenderer/Private/PreprocessApplyFilterMaterialProxy.cpp
--- a/Engine/Plugins/VirtualProduction/DMX/DMXPixelMapping/Source/DMXPixelMappingRenderer/Private/PreprocessApplyFilterMaterialProxy.cpp
+++ b/Engine/Plugins/VirtualProduction/DMX/DMXPixelMapping/Source/DMXPixelMappingRenderer/Private/PreprocessApplyFilterMaterialProxy.cpp
@@ -152,15 +152,20 @@ namespace UE::DMXPixelMapping::Rendering::Preprocess::Private
 
 	void FPreprocessApplyFilterMaterialProxy::UpdateRenderTargets(int32 NumDownsamplePasses, const TOptional<FVector2D>& OptionalOutputSize)
 	{
+		if (!WeakInputTexture.IsValid())
+		{
+			OutputRenderTarget = nullptr;
+			return;
+		}
 		UTexture* InputTexture = WeakInputTexture.Get();
+		check(InputTexture);
 
 		const bool bResize = 
 			OptionalOutputSize.IsSet() &&
-			InputTexture &&
 			InputTexture->GetSurfaceWidth() != OptionalOutputSize.GetValue().X &&
 			InputTexture->GetSurfaceHeight() != OptionalOutputSize.GetValue().Y;
 
-		const bool bNeedsAnyRenderTargets = WeakInputTexture.IsValid() && (NumDownsamplePasses > 0 || MaterialInstanceDynamic || bResize);
+		const bool bNeedsAnyRenderTargets = NumDownsamplePasses > 0 || MaterialInstanceDynamic || bResize;
 		if (!bNeedsAnyRenderTargets)
 		{
 			OutputRenderTarget = nullptr;
