DMX - Show Pixel Mapping xyY color space range as normal value instead of percents

diff --git a/Engine/Plugins/VirtualProduction/DMX/DMXPixelMapping/Source/DMXPixelMappingEditor/Private/Customizations/DMXPixelMappingColorSpaceDetails_xyY.cpp b/Engine/Plugins/VirtualProduction/DMX/DMXPixelMapping/Source/DMXPixelMappingEditor/Private/Customizations/DMXPixelMappingColorSpaceDetails_xyY.cpp
deleted file mode 100644
diff --git a/Engine/Plugins/VirtualProduction/DMX/DMXPixelMapping/Source/DMXPixelMappingEditor/Private/Customizations/DMXPixelMappingColorSpaceDetails_xyY.h b/Engine/Plugins/VirtualProduction/DMX/DMXPixelMapping/Source/DMXPixelMappingEditor/Private/Customizations/DMXPixelMappingColorSpaceDetails_xyY.h
deleted file mode 100644
diff --git a/Engine/Plugins/VirtualProduction/DMX/DMXPixelMapping/Source/DMXPixelMappingEditor/Private/Views/SDMXPixelMappingDetailsView.cpp b/Engine/Plugins/VirtualProduction/DMX/DMXPixelMapping/Source/DMXPixelMappingEditor/Private/Views/SDMXPixelMappingDetailsView.cpp
--- a/Engine/Plugins/VirtualProduction/DMX/DMXPixelMapping/Source/DMXPixelMappingEditor/Private/Views/SDMXPixelMappingDetailsView.cpp
+++ b/Engine/Plugins/VirtualProduction/DMX/DMXPixelMapping/Source/DMXPixelMappingEditor/Private/Views/SDMXPixelMappingDetailsView.cpp
@@ -3,7 +3,6 @@
 #include "Views/SDMXPixelMappingDetailsView.h"
 
 #include "ColorSpace/DMXPixelMappingColorSpace_RGBCMY.h"
-#include "ColorSpace/DMXPixelMappingColorSpace_xyY.h"
 #include "Components/DMXPixelMappingRootComponent.h"
 #include "Components/DMXPixelMappingFixtureGroupComponent.h"
 #include "Components/DMXPixelMappingFixtureGroupItemComponent.h"
@@ -12,7 +11,6 @@
 #include "Components/DMXPixelMappingMatrixCellComponent.h"
 #include "Components/DMXPixelMappingRendererComponent.h"
 #include "Customizations/DMXPixelMappingColorSpaceDetails_RGBCMY.h"
-#include "Customizations/DMXPixelMappingColorSpaceDetails_xyY.h"
 #include "Customizations/DMXPixelMappingDetailCustomization_FixtureGroup.h"
 #include "Customizations/DMXPixelMappingDetailCustomization_FixtureGroupItem.h"
 #include "Customizations/DMXPixelMappingDetailCustomization_Screen.h"
@@ -72,7 +70,6 @@ SDMXPixelMappingDetailsView::~SDMXPixelMappingDetailsView()
 		PropertyView->UnregisterInstancedCustomPropertyLayout(UDMXPixelMappingMatrixCellComponent::StaticClass());
 		
 		PropertyView->UnregisterInstancedCustomPropertyLayout(UDMXPixelMappingColorSpace_RGBCMY::StaticClass());
-		PropertyView->UnregisterInstancedCustomPropertyLayout(UDMXPixelMappingColorSpace_xyY::StaticClass());
 	}
 }
 
@@ -138,7 +135,4 @@ void SDMXPixelMappingDetailsView::RegisterCustomizations()
 
 	FOnGetDetailCustomizationInstance ColorSpaceCustomizationInstance_RGBCMY = FOnGetDetailCustomizationInstance::CreateStatic(&FDMXPixelMappingColorSpaceDetails_RGBCMY::MakeInstance);
 	PropertyView->RegisterInstancedCustomPropertyLayout(UDMXPixelMappingColorSpace_RGBCMY::StaticClass(), ColorSpaceCustomizationInstance_RGBCMY);
-
-	FOnGetDetailCustomizationInstance ColorSpaceCustomizationInstance_xyY = FOnGetDetailCustomizationInstance::CreateStatic(&FDMXPixelMappingColorSpaceDetails_xyY::MakeInstance);
-	PropertyView->RegisterInstancedCustomPropertyLayout(UDMXPixelMappingColorSpace_xyY::StaticClass(), ColorSpaceCustomizationInstance_xyY);
 }
diff --git a/Engine/Plugins/VirtualProduction/DMX/DMXPixelMapping/Source/DMXPixelMappingRuntime/Private/ColorSpace/DMXPixelMappingColorSpace_xyY.cpp b/Engine/Plugins/VirtualProduction/DMX/DMXPixelMapping/Source/DMXPixelMappingRuntime/Private/ColorSpace/DMXPixelMappingColorSpace_xyY.cpp
--- a/Engine/Plugins/VirtualProduction/DMX/DMXPixelMapping/Source/DMXPixelMappingRuntime/Private/ColorSpace/DMXPixelMappingColorSpace_xyY.cpp
+++ b/Engine/Plugins/VirtualProduction/DMX/DMXPixelMapping/Source/DMXPixelMappingRuntime/Private/ColorSpace/DMXPixelMappingColorSpace_xyY.cpp
@@ -8,9 +8,7 @@ UDMXPixelMappingColorSpace_xyY::UDMXPixelMappingColorSpace_xyY()
 	, YAttribute("CIE_Y")
 	, LuminanceAttribute("Dimmer")
 	, SRGBColorSpace(UE::Color::EColorSpace::sRGB)
-{
-	ColorSpaceRangePercents = ColorSpaceRange * 100.f;
-}
+{}
 
 void UDMXPixelMappingColorSpace_xyY::SetRGBA(const FLinearColor& InColor)
 {
@@ -47,13 +45,6 @@ void UDMXPixelMappingColorSpace_xyY::SetRGBA(const FLinearColor& InColor)
 	}
 }
 
-void UDMXPixelMappingColorSpace_xyY::PostLoad()
-{
-	Super::PostLoad();
-
-	ColorSpaceRangePercents = ColorSpaceRange * 100.f;
-}
-
 #if WITH_EDITOR
 void UDMXPixelMappingColorSpace_xyY::PostEditChangeProperty(FPropertyChangedEvent& PropertyChangedEvent)
 {
@@ -83,9 +74,5 @@ void UDMXPixelMappingColorSpace_xyY::PostEditChangeProperty(FPropertyChangedEven
 			MinLuminance = MaxLuminance;
 		}
 	}
-	else if (PropertyName == GET_MEMBER_NAME_CHECKED(UDMXPixelMappingColorSpace_xyY, ColorSpaceRangePercents))
-	{
-		ColorSpaceRange = ColorSpaceRangePercents / 100.f;
-	}
 }
 #endif // WITH_EDITOR
diff --git a/Engine/Plugins/VirtualProduction/DMX/DMXPixelMapping/Source/DMXPixelMappingRuntime/Private/ColorSpace/DMXPixelMappingColorSpace_xyY.h b/Engine/Plugins/VirtualProduction/DMX/DMXPixelMapping/Source/DMXPixelMappingRuntime/Private/ColorSpace/DMXPixelMappingColorSpace_xyY.h
--- a/Engine/Plugins/VirtualProduction/DMX/DMXPixelMapping/Source/DMXPixelMappingRuntime/Private/ColorSpace/DMXPixelMappingColorSpace_xyY.h
+++ b/Engine/Plugins/VirtualProduction/DMX/DMXPixelMapping/Source/DMXPixelMappingRuntime/Private/ColorSpace/DMXPixelMappingColorSpace_xyY.h
@@ -32,9 +32,9 @@ public:
 	UPROPERTY(EditAnywhere, Category = "XY", Meta = (DisplayName = "y Attribute"))
 	FDMXAttributeName YAttribute;
 
-	/** Sets the range of the xyY color space. CIE 1931 is 100%. */
-	UPROPERTY(EditAnywhere, Transient, Category = "XY", Meta = (ClampMin = 0.1, ClampMax = 1000.0, UIMin = 50.0, UIMax = 100.0, SliderExponent = 3, DisplayName = "Color Space Range"))
-	float ColorSpaceRangePercents = 100.f;
+	/** Sets the range of the xyY color space. A value of 1.0 corresponds to the CIE 1931 xyY . */
+	UPROPERTY(EditAnywhere, Category = "XY", Meta = (ClampMin = 0.1, ClampMax = 2.0, UIMin = 0.1, UIMax = 1.0, DisplayName = "Color Space Range"))
+	float ColorSpaceRange = .8f;
 
 	/** Attribute sent for Y */
 	UPROPERTY(EditAnywhere, Category = "Luminance")
@@ -50,16 +50,11 @@ public:
 
 protected:
 	//~ Begin UObject Interface
-	virtual void PostLoad() override;
 #if WITH_EDITOR
 	virtual void PostEditChangeProperty(FPropertyChangedEvent& PropertyChangedEvent) override;
 #endif // WITH_EDITOR
 
 private:
-	/** The actual range property */
-	UPROPERTY()
-	float ColorSpaceRange = 1.0;
-
 	/** Cached sRGB color space, to avoid instantiating on each conversion */
 	UE::Color::FColorSpace SRGBColorSpace;
 };
