DMX - Pixel Mapping DMX Library view now updates correctly when the DMX Library changed

diff --git a/Engine/Plugins/VirtualProduction/DMX/DMXPixelMapping/Source/DMXPixelMappingEditor/Private/Views/SDMXPixelMappingDMXLibraryView.cpp b/Engine/Plugins/VirtualProduction/DMX/DMXPixelMapping/Source/DMXPixelMappingEditor/Private/Views/SDMXPixelMappingDMXLibraryView.cpp
--- a/Engine/Plugins/VirtualProduction/DMX/DMXPixelMapping/Source/DMXPixelMappingEditor/Private/Views/SDMXPixelMappingDMXLibraryView.cpp
+++ b/Engine/Plugins/VirtualProduction/DMX/DMXPixelMapping/Source/DMXPixelMappingEditor/Private/Views/SDMXPixelMappingDMXLibraryView.cpp
@@ -94,6 +94,9 @@ void SDMXPixelMappingDMXLibraryView::Construct(const FArguments& InArgs, const T
 
 	ForceRefresh();
 
+	UDMXLibrary::GetOnEntitiesAdded().AddSP(this, &SDMXPixelMappingDMXLibraryView::OnEntitiesAddedOrRemoved);
+	UDMXLibrary::GetOnEntitiesRemoved().AddSP(this, &SDMXPixelMappingDMXLibraryView::OnEntitiesAddedOrRemoved);
+
 	Toolkit->GetOnSelectedComponentsChangedDelegate().AddSP(this, &SDMXPixelMappingDMXLibraryView::OnComponentSelected);
 }
 
@@ -188,6 +191,11 @@ void SDMXPixelMappingDMXLibraryView::ForceRefresh()
 	}
 }
 
+void SDMXPixelMappingDMXLibraryView::OnEntitiesAddedOrRemoved(UDMXLibrary* DMXLibrary, TArray<UDMXEntity*> Entities)
+{
+	RequestRefresh();
+}
+
 void SDMXPixelMappingDMXLibraryView::OnComponentSelected()
 {
 	const TSharedPtr<FDMXPixelMappingToolkit> Toolkit = WeakToolkit.Pin();
diff --git a/Engine/Plugins/VirtualProduction/DMX/DMXPixelMapping/Source/DMXPixelMappingEditor/Private/Views/SDMXPixelMappingDMXLibraryView.h b/Engine/Plugins/VirtualProduction/DMX/DMXPixelMapping/Source/DMXPixelMappingEditor/Private/Views/SDMXPixelMappingDMXLibraryView.h
--- a/Engine/Plugins/VirtualProduction/DMX/DMXPixelMapping/Source/DMXPixelMappingEditor/Private/Views/SDMXPixelMappingDMXLibraryView.h
+++ b/Engine/Plugins/VirtualProduction/DMX/DMXPixelMapping/Source/DMXPixelMappingEditor/Private/Views/SDMXPixelMappingDMXLibraryView.h
@@ -47,6 +47,9 @@ private:
 	/** Refreshes the widget */
 	void ForceRefresh();
 
+	/** Called when entities were added to or removed from a DMX Library */
+	void OnEntitiesAddedOrRemoved(UDMXLibrary* DMXLibrary, TArray<UDMXEntity*> Entities);
+
 	/** Called when a component was selected */
 	void OnComponentSelected();
 
