DMX - Pixel Mapping Texture now correctly hides its size

diff --git a/Engine/Plugins/VirtualProduction/DMX/DMXPixelMapping/Source/DMXPixelMappingRenderer/Private/DMXPixelMappingPreprocessRenderer.cpp b/Engine/Plugins/VirtualProduction/DMX/DMXPixelMapping/Source/DMXPixelMappingRenderer/Private/DMXPixelMappingPreprocessRenderer.cpp
--- a/Engine/Plugins/VirtualProduction/DMX/DMXPixelMapping/Source/DMXPixelMappingRenderer/Private/DMXPixelMappingPreprocessRenderer.cpp
+++ b/Engine/Plugins/VirtualProduction/DMX/DMXPixelMapping/Source/DMXPixelMappingRenderer/Private/DMXPixelMappingPreprocessRenderer.cpp
@@ -17,6 +17,8 @@ void UDMXPixelMappingPreprocessRenderer::SetInputTexture(UTexture* InTexture)
 	using namespace UE::DMXPixelMapping::Rendering::Preprocess::Private;
 	RenderInputProxy = MakeShared<FPreprocessRenderInputTextureProxy>(InTexture);
 	ApplyMaterialProxy = MakeShared<FPreprocessApplyFilterMaterialProxy>();
+
+	bShowInputSize = false;
 }
 
 void UDMXPixelMappingPreprocessRenderer::SetInputMaterial(UMaterialInterface* InMaterial)
@@ -24,6 +26,8 @@ void UDMXPixelMappingPreprocessRenderer::SetInputMaterial(UMaterialInterface* In
 	using namespace UE::DMXPixelMapping::Rendering::Preprocess::Private;
 	RenderInputProxy = MakeShared<FPreprocessRenderInputMaterialProxy>(InMaterial, InputSize);
 	ApplyMaterialProxy = MakeShared<FPreprocessApplyFilterMaterialProxy>();
+
+	bShowInputSize = true;
 }
 
 void UDMXPixelMappingPreprocessRenderer::SetInputUserWidget(UUserWidget* InUserWidget)
@@ -31,6 +35,8 @@ void UDMXPixelMappingPreprocessRenderer::SetInputUserWidget(UUserWidget* InUserW
 	using namespace UE::DMXPixelMapping::Rendering::Preprocess::Private;
 	RenderInputProxy = MakeShared<FPreprocessRenderInputUserWidgetProxy>(InUserWidget, InputSize);
 	ApplyMaterialProxy = MakeShared<FPreprocessApplyFilterMaterialProxy>();
+
+	bShowInputSize = true;
 }
 
 void UDMXPixelMappingPreprocessRenderer::ClearInput()
diff --git a/Engine/Plugins/VirtualProduction/DMX/DMXPixelMapping/Source/DMXPixelMappingRenderer/Private/PreprocessApplyFilterMaterialProxy.cpp b/Engine/Plugins/VirtualProduction/DMX/DMXPixelMapping/Source/DMXPixelMappingRenderer/Private/PreprocessApplyFilterMaterialProxy.cpp
--- a/Engine/Plugins/VirtualProduction/DMX/DMXPixelMapping/Source/DMXPixelMappingRenderer/Private/PreprocessApplyFilterMaterialProxy.cpp
+++ b/Engine/Plugins/VirtualProduction/DMX/DMXPixelMapping/Source/DMXPixelMappingRenderer/Private/PreprocessApplyFilterMaterialProxy.cpp
@@ -183,7 +183,7 @@ namespace UE::DMXPixelMapping::Rendering::Preprocess::Private
 		{
 			FlushRenderingCommands();
 
-			const FVector2D InputSize = FVector2D(InputTexture->GetSurfaceHeight(), InputTexture->GetSurfaceWidth());
+			const FVector2D InputSize = FVector2D(InputTexture->GetSurfaceWidth(), InputTexture->GetSurfaceHeight());
 
 			// Create downsample render targets
 			DownsampleRenderTargets.Reset();
diff --git a/Engine/Plugins/VirtualProduction/DMX/DMXPixelMapping/Source/DMXPixelMappingRenderer/Public/DMXPixelMappingPreprocessRenderer.h b/Engine/Plugins/VirtualProduction/DMX/DMXPixelMapping/Source/DMXPixelMappingRenderer/Public/DMXPixelMappingPreprocessRenderer.h
--- a/Engine/Plugins/VirtualProduction/DMX/DMXPixelMapping/Source/DMXPixelMappingRenderer/Public/DMXPixelMappingPreprocessRenderer.h
+++ b/Engine/Plugins/VirtualProduction/DMX/DMXPixelMapping/Source/DMXPixelMappingRenderer/Public/DMXPixelMappingPreprocessRenderer.h
@@ -121,7 +121,7 @@ protected:
 
 private:
 	/** Size of the input texture (Material and UMG only) */
-	UPROPERTY(EditAnywhere, Category = "Render Settings", Meta = (AllowPrivateAccess = true))
+	UPROPERTY(EditAnywhere, Category = "Render Settings", Meta = (EditCondition = "bShowInputSize", EditConditionHides, AllowPrivateAccess = true))
 	FVector2D InputSize{ 1024.f, 1024.f };
 
 	/** Number of times the pixelmapping input is downsampled */
@@ -133,7 +133,7 @@ private:
 	EDMXPixelMappingRenderingPreprocessorSizeMode OutputSizeMode = EDMXPixelMappingRenderingPreprocessorSizeMode::SameAsInput;
 
 	/** Size of the rendered texture */
-	UPROPERTY(EditAnywhere, Category = "Filtering", Meta = (DisplayName = "Custom Size", ClampMin = "1", ClampMax = 2048, UIMin = 1, UIMax = 2048, EditCondition = "OutputSizeMode == EDMXPixelMappingRenderingPreprocessorSizeMode::Custom", EditConditionHides, AllowPrivateAccess = true))
+	UPROPERTY(EditAnywhere, Category = "Filtering", Meta = (DisplayName = "Custom Size", ClampMin = "1", ClampMax = 2048, UIMin = 1, UIMax = 2048, EditCondition = "OutputSizeMode == EDMXPixelMappingRenderingPreprocessorSizeMode::CustomSize", EditConditionHides, AllowPrivateAccess = true))
 	FIntPoint CustomOutputSize { 256, 256 };
 	
 	/** Filter material applied to the rendered input */
@@ -150,7 +150,7 @@ private:
 
 	/** Blur distance applied, only applicable if the filter matierial has a "BlurDistance" parameter */
 	UPROPERTY(EditAnywhere, Category = "Filtering", Meta = (UIMin = "0", UIMax = "1", SliderExponent = 2.0, AllowPrivateAccess = true))
-	float BlurDistance = 0.02f;
+	float BlurDistance = .02f;
 
 	/** Parameter name of the filter material to receive the input texture */
 	UPROPERTY(EditAnywhere, AdvancedDisplay, Category = "Filtering", Meta = (AllowPrivateAccess = true))
@@ -160,6 +160,10 @@ private:
 	UPROPERTY(EditAnywhere, AdvancedDisplay, Category = "Filtering", Meta = (AllowPrivateAccess = true))
 	FName BlurDistanceParamterName = "BlurDistance";
 
+	/** If true, shows the input size property */
+	UPROPERTY(Transient)
+	bool bShowInputSize = false;
+
 	/** Proxy to render the input */
 	TSharedPtr<UE::DMXPixelMapping::Rendering::Preprocess::Private::IPreprocessRenderInputProxy> RenderInputProxy;
 
