DMX - Fix Pixel Mapping Preprocess Renderer input size no longer shows its edit condition, specifies a reasonable range

diff --git a/Engine/Plugins/VirtualProduction/DMX/DMXPixelMapping/Source/DMXPixelMappingRenderer/Public/DMXPixelMappingPreprocessRenderer.h b/Engine/Plugins/VirtualProduction/DMX/DMXPixelMapping/Source/DMXPixelMappingRenderer/Public/DMXPixelMappingPreprocessRenderer.h
--- a/Engine/Plugins/VirtualProduction/DMX/DMXPixelMapping/Source/DMXPixelMappingRenderer/Public/DMXPixelMappingPreprocessRenderer.h
+++ b/Engine/Plugins/VirtualProduction/DMX/DMXPixelMapping/Source/DMXPixelMappingRenderer/Public/DMXPixelMappingPreprocessRenderer.h
@@ -120,8 +120,8 @@ protected:
 	//~ End UObject interface
 
 private:
-	/** Size of the input texture (Material and UMG only) */
-	UPROPERTY(EditAnywhere, Category = "Render Settings", Meta = (EditCondition = "bShowInputSize", EditConditionHides, AllowPrivateAccess = true))
+	/** Size of the input texture (Material and User Widget only) */
+	UPROPERTY(EditAnywhere, Category = "Render Settings", Meta = (EditCondition = "bShowInputSize", EditConditionHides, HideEditConditionToggle, ClampMin = "1", ClampMax = "8192", UIMin = "1", UIMax = "4096", AllowPrivateAccess = true))
 	FVector2D InputSize{ 1024.f, 1024.f };
 
 	/** Number of times the pixelmapping input is downsampled */
