DMX - PixelMapping now only draws the visible part of the texture when zoomed in, fixes issues and allows for higher zoom levels

diff --git a/Engine/Plugins/VirtualProduction/DMX/DMXPixelMapping/Source/DMXPixelMappingEditor/Private/Customizations/DMXPixelMappingDetailCustomization_Renderer.cpp b/Engine/Plugins/VirtualProduction/DMX/DMXPixelMapping/Source/DMXPixelMappingEditor/Private/Customizations/DMXPixelMappingDetailCustomization_Renderer.cpp
--- a/Engine/Plugins/VirtualProduction/DMX/DMXPixelMapping/Source/DMXPixelMappingEditor/Private/Customizations/DMXPixelMappingDetailCustomization_Renderer.cpp
+++ b/Engine/Plugins/VirtualProduction/DMX/DMXPixelMapping/Source/DMXPixelMappingEditor/Private/Customizations/DMXPixelMappingDetailCustomization_Renderer.cpp
@@ -199,7 +199,7 @@ EVisibility FDMXPixelMappingDetailCustomization_Renderer::GetInputTextureWarning
 {
 	if (UDMXPixelMappingRendererComponent* Component = RendererComponent.Get())
 	{
-		if (Component->GetRendererInputTexture() == nullptr)
+		if (!Component->GetRenderedInputTexture())
 		{
 			return EVisibility::Visible;
 		}
@@ -212,7 +212,7 @@ FText FDMXPixelMappingDetailCustomization_Renderer::GetInputTextureWarningText()
 {
 	if (UDMXPixelMappingRendererComponent* Component = RendererComponent.Get())
 	{
-		if (Component->GetRendererInputTexture() == nullptr)
+		if (!Component->GetRenderedInputTexture())
 		{
 			if (Component->RendererType == EDMXPixelMappingRendererType::Texture)
 			{
diff --git a/Engine/Plugins/VirtualProduction/DMX/DMXPixelMapping/Source/DMXPixelMappingEditor/Private/Toolkits/DMXPixelMappingToolkit.cpp b/Engine/Plugins/VirtualProduction/DMX/DMXPixelMapping/Source/DMXPixelMappingEditor/Private/Toolkits/DMXPixelMappingToolkit.cpp
--- a/Engine/Plugins/VirtualProduction/DMX/DMXPixelMapping/Source/DMXPixelMappingEditor/Private/Toolkits/DMXPixelMappingToolkit.cpp
+++ b/Engine/Plugins/VirtualProduction/DMX/DMXPixelMapping/Source/DMXPixelMappingEditor/Private/Toolkits/DMXPixelMappingToolkit.cpp
@@ -413,16 +413,7 @@ void FDMXPixelMappingToolkit::DeleteSelectedComponents()
 
 		UpdateBlueprintNodes(DMXPixelMapping);
 	}
-}
-
-void FDMXPixelMappingToolkit::ZoomToFit()
-{
-	if (DesignerView.IsValid())
-	{
-		constexpr bool bInstantZoom = false;
-		DesignerView->ZoomToFit(bInstantZoom);
-	}
-}
+}	
 
 bool FDMXPixelMappingToolkit::CanSizeSelectedComponentToTexture() const
 {
@@ -743,13 +734,13 @@ TSharedRef<SDockTab> FDMXPixelMappingToolkit::SpawnTab_LayoutView(const FSpawnTa
 
 void FDMXPixelMappingToolkit::CreateInternalViews()
 {
-	CreateOrGetView_InputSourceView();
-	CreateOrGetView_DMXLibraryView();
-	CreateOrGetView_HierarchyView();
-	CreateOrGetView_DesignerView();
-	CreateOrGetView_PreviewView();
-	CreateOrGetView_DetailsView();
-	CreateOrGetView_LayoutView();
+	GetOrCreateInputSourceView();
+	GetOrCreateDMXLibraryView();
+	GetOrCreateHierarchyView();
+	GetOrCreateDesignerView();
+	GetOrCreatePreviewView();
+	GetOrCreateDetailsView();
+	GetOrCreateLayoutView();
 }
 
 void FDMXPixelMappingToolkit::OnComponentRenamed(UDMXPixelMappingBaseComponent* InComponent)
@@ -757,7 +748,7 @@ void FDMXPixelMappingToolkit::OnComponentRenamed(UDMXPixelMappingBaseComponent*
 	UpdateBlueprintNodes(GetDMXPixelMapping());
 }
 
-TSharedRef<SWidget> FDMXPixelMappingToolkit::CreateOrGetView_InputSourceView()
+TSharedRef<SDMXPixelMappingInputSourceView> FDMXPixelMappingToolkit::GetOrCreateInputSourceView()
 {
 	if (!InputSourceView.IsValid())
 	{
@@ -767,7 +758,7 @@ TSharedRef<SWidget> FDMXPixelMappingToolkit::CreateOrGetView_InputSourceView()
 	return InputSourceView.ToSharedRef();
 }
 
-TSharedRef<SWidget> FDMXPixelMappingToolkit::CreateOrGetView_DMXLibraryView()
+TSharedRef<SDMXPixelMappingDMXLibraryView> FDMXPixelMappingToolkit::GetOrCreateDMXLibraryView()
 {
 	if (!DMXLibraryView.IsValid())
 	{
@@ -777,7 +768,7 @@ TSharedRef<SWidget> FDMXPixelMappingToolkit::CreateOrGetView_DMXLibraryView()
 	return DMXLibraryView.ToSharedRef();
 }
 
-TSharedRef<SWidget> FDMXPixelMappingToolkit::CreateOrGetView_HierarchyView()
+TSharedRef<SDMXPixelMappingHierarchyView> FDMXPixelMappingToolkit::GetOrCreateHierarchyView()
 {
 	if (!HierarchyView.IsValid())
 	{
@@ -787,7 +778,7 @@ TSharedRef<SWidget> FDMXPixelMappingToolkit::CreateOrGetView_HierarchyView()
 	return HierarchyView.ToSharedRef();
 }
 
-TSharedRef<SWidget> FDMXPixelMappingToolkit::CreateOrGetView_DesignerView()
+TSharedRef<SDMXPixelMappingDesignerView> FDMXPixelMappingToolkit::GetOrCreateDesignerView()
 {
 	if (!DesignerView.IsValid())
 	{
@@ -797,7 +788,7 @@ TSharedRef<SWidget> FDMXPixelMappingToolkit::CreateOrGetView_DesignerView()
 	return DesignerView.ToSharedRef();
 }
 
-TSharedRef<SWidget> FDMXPixelMappingToolkit::CreateOrGetView_PreviewView()
+TSharedRef<SDMXPixelMappingPreviewView> FDMXPixelMappingToolkit::GetOrCreatePreviewView()
 {
 	if (!PreviewView.IsValid())
 	{
@@ -807,7 +798,7 @@ TSharedRef<SWidget> FDMXPixelMappingToolkit::CreateOrGetView_PreviewView()
 	return PreviewView.ToSharedRef();
 }
 
-TSharedRef<SWidget> FDMXPixelMappingToolkit::CreateOrGetView_DetailsView()
+TSharedRef<SDMXPixelMappingDetailsView> FDMXPixelMappingToolkit::GetOrCreateDetailsView()
 {
 	if (!DetailsView.IsValid())
 	{
@@ -817,7 +808,7 @@ TSharedRef<SWidget> FDMXPixelMappingToolkit::CreateOrGetView_DetailsView()
 	return DetailsView.ToSharedRef();
 }
 
-TSharedRef<SWidget> FDMXPixelMappingToolkit::CreateOrGetView_LayoutView()
+TSharedRef<SDMXPixelMappingLayoutView> FDMXPixelMappingToolkit::GetOrCreateLayoutView()
 {
 	if (!LayoutView.IsValid())
 	{
diff --git a/Engine/Plugins/VirtualProduction/DMX/DMXPixelMapping/Source/DMXPixelMappingEditor/Private/Toolkits/DMXPixelMappingToolkit.h b/Engine/Plugins/VirtualProduction/DMX/DMXPixelMapping/Source/DMXPixelMappingEditor/Private/Toolkits/DMXPixelMappingToolkit.h
--- a/Engine/Plugins/VirtualProduction/DMX/DMXPixelMapping/Source/DMXPixelMappingEditor/Private/Toolkits/DMXPixelMappingToolkit.h
+++ b/Engine/Plugins/VirtualProduction/DMX/DMXPixelMapping/Source/DMXPixelMappingEditor/Private/Toolkits/DMXPixelMappingToolkit.h
@@ -107,25 +107,25 @@ public:
 	const TSet<FDMXPixelMappingComponentReference>& GetSelectedComponents() const { return SelectedComponents; }
 
 	/** Gets or creates the Input Source View for this Pixel Mapping instance */
-	TSharedRef<SWidget> CreateOrGetView_InputSourceView();
+	TSharedRef<SDMXPixelMappingInputSourceView> GetOrCreateInputSourceView();
 
 	/** Gets or creates the DMX Library View for this Pixel Mapping instance */
-	TSharedRef<SWidget> CreateOrGetView_DMXLibraryView();
+	TSharedRef<SDMXPixelMappingDMXLibraryView> GetOrCreateDMXLibraryView();
 
 	/** Gets or creates the Hierarchy View for this Pixel Mapping instance */
-	TSharedRef<SWidget> CreateOrGetView_HierarchyView();
+	TSharedRef<SDMXPixelMappingHierarchyView> GetOrCreateHierarchyView();
 
 	/** Gets or creates the Designer View for this Pixel Mapping instance */
-	TSharedRef<SWidget> CreateOrGetView_DesignerView();
+	TSharedRef<SDMXPixelMappingDesignerView> GetOrCreateDesignerView();
 
 	/** Gets or creates the Preview View for this Pixel Mapping instance */
-	TSharedRef<SWidget> CreateOrGetView_PreviewView();
+	TSharedRef<SDMXPixelMappingPreviewView> GetOrCreatePreviewView();
 
 	/** Gets or creates the Details View for this Pixel Mapping instance */
-	TSharedRef<SWidget> CreateOrGetView_DetailsView();
+	TSharedRef<SDMXPixelMappingDetailsView> GetOrCreateDetailsView();
 
 	/** Gets or creates the Layout View for this Pixel Mapping instance */
-	TSharedRef<SWidget> CreateOrGetView_LayoutView();
+	TSharedRef<SDMXPixelMappingLayoutView> GetOrCreateLayoutView();
 
 	bool IsPlayingDMX() const { return bIsPlayingDMX; }
 
@@ -154,9 +154,6 @@ public:
 	/** Deletes the selected Components */
 	void DeleteSelectedComponents();
 
-	/** Zooms to fit */
-	void ZoomToFit();
-
 	/** Returns true if the selected component can be sized to texture */
 	bool CanSizeSelectedComponentToTexture() const;
 
diff --git a/Engine/Plugins/VirtualProduction/DMX/DMXPixelMapping/Source/DMXPixelMappingEditor/Private/Viewports/DMXPixelMappingSourceTextureViewportClient.cpp b/Engine/Plugins/VirtualProduction/DMX/DMXPixelMapping/Source/DMXPixelMappingEditor/Private/Viewports/DMXPixelMappingSourceTextureViewportClient.cpp
--- a/Engine/Plugins/VirtualProduction/DMX/DMXPixelMapping/Source/DMXPixelMappingEditor/Private/Viewports/DMXPixelMappingSourceTextureViewportClient.cpp
+++ b/Engine/Plugins/VirtualProduction/DMX/DMXPixelMapping/Source/DMXPixelMappingEditor/Private/Viewports/DMXPixelMappingSourceTextureViewportClient.cpp
@@ -1,46 +1,125 @@
 // Copyright Epic Games, Inc. All Rights Reserved.
 
 #include "Viewports/DMXPixelMappingSourceTextureViewportClient.h"
-#include "Engine/Texture.h"
-#include "TextureResource.h"
-#include "Widgets/SDMXPixelMappingSourceTextureViewport.h"
-#include "Components/DMXPixelMappingRendererComponent.h"
-#include "Toolkits/DMXPixelMappingToolkit.h"
 
 #include "CanvasItem.h"
 #include "CanvasTypes.h"
+#include "Components/DMXPixelMappingRendererComponent.h"
+#include "Engine/Texture.h"
 #include "Slate/SceneViewport.h"
+#include "Texture2DPreview.h"
+#include "TextureResource.h"
+#include "Toolkits/DMXPixelMappingToolkit.h"
+#include "Views/SDMXPixelMappingDesignerView.h"
+#include "Widgets/SDMXPixelMappingSourceTextureViewport.h"
+
+
+FDMXPixelMappingSourceTextureViewportClient::FDMXPixelMappingSourceTextureViewportClient(const TSharedPtr<FDMXPixelMappingToolkit>& InToolkit, TWeakPtr<SDMXPixelMappingSourceTextureViewport> InSourceTextureViewport)
+	: WeakToolkit(InToolkit)
+	, WeakSourceTextureViewport(InSourceTextureViewport)
+{}
+
+
+bool FDMXPixelMappingSourceTextureViewportClient::DrawOnlyVisibleRect() const
+{
+	if (!WeakToolkit.IsValid())
+	{
+		return 0.0;
+	}
+	const TSharedRef<FDMXPixelMappingToolkit> Toolkit = WeakToolkit.Pin().ToSharedRef();
+	const TSharedRef<SDMXPixelMappingDesignerView> DesignerView = Toolkit->GetOrCreateDesignerView();
+
+	UDMXPixelMappingRendererComponent* RendererComponent = Toolkit->GetActiveRendererComponent();
+	UTexture* InputTexture = RendererComponent ? RendererComponent->GetRenderedInputTexture() : nullptr;
+	if (!InputTexture)
+	{
+		return false;
+	}
+
+	const int32 ExcessThreshold = GMaxTextureDimensions / 4;
 
-FDMXPixelMappingSourceTextureViewportClient::FDMXPixelMappingSourceTextureViewportClient(const TSharedPtr<FDMXPixelMappingToolkit>& InToolkit, TWeakPtr<SDMXPixelMappingSourceTextureViewport> InViewport)
-	: ToolkitWeakPtr(InToolkit)
-	, WeakViewport(InViewport)
+	return
+		InputTexture->GetSurfaceWidth() * DesignerView->GetZoomAmount() > ExcessThreshold ||
+		InputTexture->GetSurfaceHeight() * DesignerView->GetZoomAmount() > ExcessThreshold;
+}
+
+FBox2D FDMXPixelMappingSourceTextureViewportClient::GetVisibleTextureBoxGraphSpace() const
 {
-	check(ToolkitWeakPtr.IsValid() && WeakViewport.IsValid());
+	if (!WeakToolkit.IsValid())
+	{
+		return FBox2D();
+	}
+	const TSharedRef<FDMXPixelMappingToolkit> Toolkit = WeakToolkit.Pin().ToSharedRef();
+	const TSharedRef<SDMXPixelMappingDesignerView> DesignerView = Toolkit->GetOrCreateDesignerView();
+
+	UDMXPixelMappingRendererComponent* RendererComponent = Toolkit->GetActiveRendererComponent();
+	UTexture* InputTexture = RendererComponent ? RendererComponent->GetRenderedInputTexture() : nullptr;
+	if (!InputTexture || !InputTexture->GetResource())
+	{
+		return FBox2D();
+	}
+
+	// Compute texture box in graph space
+	const FVector2D TextureDimensions = FVector2D(InputTexture->GetSurfaceWidth(), InputTexture->GetSurfaceHeight());
+	const FBox2D TextureBox(FVector2D::ZeroVector, TextureDimensions);
+
+	// Compute graph box in graph space
+	const FGeometry& GraphGeometry = DesignerView->GetGraphTickSpaceGeometry();
+	const FVector2D GraphSize = GraphGeometry.GetAbsoluteSize() / DesignerView->GetZoomAmount();
+	const FVector2D ViewOffset = DesignerView->GetViewOffset();
+	const FBox2D GraphBox(ViewOffset, ViewOffset + GraphSize);
+
+	return TextureBox.Overlap(GraphBox);
 }
 
 void FDMXPixelMappingSourceTextureViewportClient::Draw(FViewport* InViewport, FCanvas* Canvas)
 {
 	Canvas->Clear(FColor::Transparent);
 
-	TSharedPtr<SDMXPixelMappingSourceTextureViewport> Viewport = WeakViewport.Pin();
-	check(Viewport.IsValid());
-	TSharedPtr<FDMXPixelMappingToolkit> Toolkit = ToolkitWeakPtr.Pin();
-	check(Toolkit.IsValid());
-	const TSharedPtr<FSceneViewport> SceneViewport = Viewport->GetViewport();
-	check(SceneViewport.IsValid());
+	if (!WeakToolkit.IsValid())
+	{
+		return;
+	}
+	const TSharedRef<FDMXPixelMappingToolkit> Toolkit = WeakToolkit.Pin().ToSharedRef();
+	const TSharedRef<SDMXPixelMappingDesignerView> DesignerView = Toolkit->GetOrCreateDesignerView();
 
-	UTexture* InputTexture = nullptr;
-	if (UDMXPixelMappingRendererComponent* RendererComponent = Toolkit->GetActiveRendererComponent())
+	UDMXPixelMappingRendererComponent* RendererComponent = Toolkit->GetActiveRendererComponent();
+	UTexture* InputTexture = RendererComponent ? RendererComponent->GetRenderedInputTexture() : nullptr;
+	if (!InputTexture || !InputTexture->GetResource())
 	{
-		InputTexture = RendererComponent->GetRendererInputTexture();
+		return;
 	}
 
-	if (InputTexture && InputTexture->GetResource())
+	const TSharedPtr<SDMXPixelMappingSourceTextureViewport> ViewportWidget = WeakSourceTextureViewport.Pin();
+	if (!ViewportWidget.IsValid())
+	{
+		return;
+	}
+
+	// Compute texture box in graph space
+	const FVector2D TextureDimensions = FVector2D(InputTexture->GetSurfaceWidth(), InputTexture->GetSurfaceHeight());
+	const FBox2D TextureBox(FVector2D::ZeroVector, TextureDimensions);
+
+	// Get the visible rect in grap space, and the UVs
+	const FBox2D VisibleRectBox = GetVisibleTextureBoxGraphSpace();
+	const FVector2D UV0 = VisibleRectBox.Min / TextureBox.GetSize();
+	const FVector2D UV1 = VisibleRectBox.Max / TextureBox.GetSize();
+
+	if (DrawOnlyVisibleRect())
 	{
-		const uint32 Width = SceneViewport->GetSizeXY().X;
-		const uint32 Height = SceneViewport->GetSizeXY().Y;
+		const TSharedPtr<FSceneViewport> SceneViewport = ViewportWidget->GetViewport();
+		if (!SceneViewport.IsValid())
+		{
+			return;
+		}
 
-		FCanvasTileItem TileItem(FVector2D(0, 0), InputTexture->GetResource(), FVector2D(Width, Height), FLinearColor::White);
+		FCanvasTileItem TileItem(FVector2D::ZeroVector, InputTexture->GetResource(), SceneViewport->GetSizeXY(), UV0, UV1, FLinearColor::White);
+		TileItem.BlendMode = ESimpleElementBlendMode::SE_BLEND_MAX;
+		Canvas->DrawItem(TileItem);
+	}
+	else
+	{
+		FCanvasTileItem TileItem(VisibleRectBox.Min * DesignerView->GetZoomAmount(), InputTexture->GetResource(), VisibleRectBox.GetSize() * DesignerView->GetZoomAmount(), UV0, UV1, FLinearColor::White);
 		TileItem.BlendMode = ESimpleElementBlendMode::SE_BLEND_MAX;
 		Canvas->DrawItem(TileItem);
 	}
diff --git a/Engine/Plugins/VirtualProduction/DMX/DMXPixelMapping/Source/DMXPixelMappingEditor/Private/Viewports/DMXPixelMappingSourceTextureViewportClient.h b/Engine/Plugins/VirtualProduction/DMX/DMXPixelMapping/Source/DMXPixelMappingEditor/Private/Viewports/DMXPixelMappingSourceTextureViewportClient.h
--- a/Engine/Plugins/VirtualProduction/DMX/DMXPixelMapping/Source/DMXPixelMappingEditor/Private/Viewports/DMXPixelMappingSourceTextureViewportClient.h
+++ b/Engine/Plugins/VirtualProduction/DMX/DMXPixelMapping/Source/DMXPixelMappingEditor/Private/Viewports/DMXPixelMappingSourceTextureViewportClient.h
@@ -4,6 +4,7 @@
 
 #include "ViewportClient.h"
 
+
 class SDMXPixelMappingSourceTextureViewport;
 class FDMXPixelMappingToolkit;
 
@@ -14,12 +15,21 @@ public:
 	/** Constructor */
 	FDMXPixelMappingSourceTextureViewportClient(const TSharedPtr<FDMXPixelMappingToolkit>& Toolkit, TWeakPtr<SDMXPixelMappingSourceTextureViewport> InViewport);
 	
+	/**
+	 * Returns true if it only the visible rectangle is drawn, instead of the whole texture.
+	 * This is favorable when zoomed in closely and the viewport size gets close to, or exceeds GMaxTextureDimensions.
+	 */
+	bool DrawOnlyVisibleRect() const;
+
+	/** Returns the visible texture size in graph space */
+	FBox2D GetVisibleTextureBoxGraphSpace() const;
+
+protected:
 	//~ Begin FViewportClient Interface
 	virtual void Draw(FViewport* Viewport, FCanvas* Canvas) override;
 	//~ End FViewportClient Interface
 
 private:
-	TWeakPtr<FDMXPixelMappingToolkit> ToolkitWeakPtr;
-
-	TWeakPtr<SDMXPixelMappingSourceTextureViewport> WeakViewport;
+	TWeakPtr<FDMXPixelMappingToolkit> WeakToolkit;
+	TWeakPtr<SDMXPixelMappingSourceTextureViewport> WeakSourceTextureViewport;
 };
diff --git a/Engine/Plugins/VirtualProduction/DMX/DMXPixelMapping/Source/DMXPixelMappingEditor/Private/Views/SDMXPixelMappingDesignerView.cpp b/Engine/Plugins/VirtualProduction/DMX/DMXPixelMapping/Source/DMXPixelMappingEditor/Private/Views/SDMXPixelMappingDesignerView.cpp
--- a/Engine/Plugins/VirtualProduction/DMX/DMXPixelMapping/Source/DMXPixelMappingEditor/Private/Views/SDMXPixelMappingDesignerView.cpp
+++ b/Engine/Plugins/VirtualProduction/DMX/DMXPixelMapping/Source/DMXPixelMappingEditor/Private/Views/SDMXPixelMappingDesignerView.cpp
@@ -30,12 +30,8 @@
 #include "Widgets/Layout/SSpacer.h"
 #include "Widgets/Input/SButton.h"
 
-#define LOCTEXT_NAMESPACE "SDMXPixelMappingDesignerView"
 
-SDMXPixelMappingDesignerView::~SDMXPixelMappingDesignerView()
-{
-	GEditor->UnregisterForUndo(this);
-}
+#define LOCTEXT_NAMESPACE "SDMXPixelMappingDesignerView"
 
 void SDMXPixelMappingDesignerView::Construct(const FArguments& InArgs, const TSharedPtr<FDMXPixelMappingToolkit>& InToolkit)
 {
@@ -122,7 +118,7 @@ void SDMXPixelMappingDesignerView::Construct(const FArguments& InArgs, const TSh
 					.HAlign(HAlign_Fill)
 					.VAlign(VAlign_Fill)
 					[
-						SNew(SDMXPixelMappingZoomPan)
+						SAssignNew(ZoomPan, SDMXPixelMappingZoomPan)
 						.ZoomAmount(this, &SDMXPixelMappingDesignerView::GetZoomAmount)
 						.ViewOffset(this, &SDMXPixelMappingDesignerView::GetViewOffset)
 						.Visibility(this, &SDMXPixelMappingDesignerView::GetZoomPanVisibility)
@@ -239,31 +235,11 @@ void SDMXPixelMappingDesignerView::Construct(const FArguments& InArgs, const TSh
 	// Bind to component changes
 	UDMXPixelMappingBaseComponent::GetOnComponentAdded().AddSP(this, &SDMXPixelMappingDesignerView::OnComponentAdded);
 	UDMXPixelMappingBaseComponent::GetOnComponentRemoved().AddSP(this, &SDMXPixelMappingDesignerView::OnComponentRemoved);
-
-	GEditor->RegisterForUndo(this);
-}
-
-FOptionalSize SDMXPixelMappingDesignerView::GetPreviewAreaWidth() const
-{
-	FVector2D Area;
-	FVector2D Size;
-	GetPreviewAreaAndSize(Area, Size);
-
-	return Area.X;
 }
 
-FOptionalSize SDMXPixelMappingDesignerView::GetPreviewAreaHeight() const
+const FGeometry& SDMXPixelMappingDesignerView::GetGraphTickSpaceGeometry() const
 {
-	FVector2D Area;
-	FVector2D Size;
-	GetPreviewAreaAndSize(Area, Size);
-
-	return Area.Y;
-}
-
-float SDMXPixelMappingDesignerView::GetPreviewScale() const
-{
-	return GetZoomAmount();
+	return ZoomPan->GetTickSpaceGeometry();
 }
 
 FReply SDMXPixelMappingDesignerView::OnMouseButtonDown(const FGeometry& MyGeometry, const FPointerEvent& MouseEvent)
@@ -460,8 +436,8 @@ void SDMXPixelMappingDesignerView::Tick(const FGeometry& AllottedGeometry, const
 		GridOrigin = AbsoluteOrigin;
 
 		// Roler position
-		TopRuler->SetRuling(AbsoluteOrigin, 1.0f / GetPreviewScale());
-		SideRuler->SetRuling(AbsoluteOrigin, 1.0f / GetPreviewScale());
+		TopRuler->SetRuling(AbsoluteOrigin, 1.0f / GetZoomAmount());
+		SideRuler->SetRuling(AbsoluteOrigin, 1.0f / GetZoomAmount());
 
 		if (IsHovered())
 		{
@@ -719,33 +695,35 @@ void SDMXPixelMappingDesignerView::CreateExtensionWidgetsForSelection()
 	// Remove all the current extension widgets
 	ClearExtensionWidgets();
 
-	// Get the selected widgets as an array
-	const TArray<FDMXPixelMappingComponentReference>& SelectedComponents = GetSelectedComponents().Array();
+	// Create new handles if possible
+	const TSet<FDMXPixelMappingComponentReference>& SelectedComponents = GetSelectedComponents();
+	if (SelectedComponents.Num() != 1)
+	{
+		return;
+	}
 
-	if (SelectedComponents.Num() == 1)
+	UDMXPixelMappingOutputComponent* OutputComponent = Cast<UDMXPixelMappingOutputComponent>(SelectedComponents.Array()[0].GetComponent());
+	if (OutputComponent && 
+		OutputComponent->GetClass() != UDMXPixelMappingRendererComponent::StaticClass() &&
+		OutputComponent->IsVisible() &&
+		!OutputComponent->IsLockInDesigner())
 	{
-		if (UDMXPixelMappingOutputComponent* OutputComponent = Cast<UDMXPixelMappingOutputComponent>(SelectedComponents[0].GetComponent()))
+		// Add transform handles
+		constexpr float Offset = 10.f;
+		const TSharedPtr<SDMXPixelMappingDesignerView> Self = SharedThis(this);
+		TransformHandles.Add(SNew(SDMXPixelMappingTransformHandle, Self, EDMXPixelMappingTransformDirection::CenterRight, FVector2D(Offset, 0.f)));
+		TransformHandles.Add(SNew(SDMXPixelMappingTransformHandle, Self, EDMXPixelMappingTransformDirection::BottomCenter, FVector2D(0.f, Offset)));
+		TransformHandles.Add(SNew(SDMXPixelMappingTransformHandle, Self, EDMXPixelMappingTransformDirection::BottomRight, FVector2D(Offset, Offset)));
+
+		// Add Widgets to designer surface
+		for (TSharedPtr<SDMXPixelMappingTransformHandle>& Handle : TransformHandles)
 		{
-			if (OutputComponent->IsVisible() && !OutputComponent->IsLockInDesigner())
-			{
-				// Add transform handles
-				constexpr float Offset = 10.f;
-				const TSharedPtr<SDMXPixelMappingDesignerView> Self = SharedThis(this);
-				TransformHandles.Add(SNew(SDMXPixelMappingTransformHandle, Self, EDMXPixelMappingTransformDirection::CenterRight, FVector2D(Offset, 0.f)));
-				TransformHandles.Add(SNew(SDMXPixelMappingTransformHandle, Self, EDMXPixelMappingTransformDirection::BottomCenter, FVector2D(0.f, Offset)));
-				TransformHandles.Add(SNew(SDMXPixelMappingTransformHandle, Self, EDMXPixelMappingTransformDirection::BottomRight, FVector2D(Offset, Offset)));
-
-				// Add Widgets to designer surface
-				for (TSharedPtr<SDMXPixelMappingTransformHandle>& Handle : TransformHandles)
-				{
-					ExtensionWidgetCanvas->AddSlot()
-						.Position(TAttribute<FVector2D>::Create(TAttribute<FVector2D>::FGetter::CreateSP(this, &SDMXPixelMappingDesignerView::GetExtensionPosition, Handle)))
-						.Size(TAttribute<FVector2D>::Create(TAttribute<FVector2D>::FGetter::CreateSP(this, &SDMXPixelMappingDesignerView::GetExtensionSize, Handle)))
-						[
-							Handle.ToSharedRef()
-						];
-				}
-			}
+			ExtensionWidgetCanvas->AddSlot()
+				.Position(TAttribute<FVector2D>::Create(TAttribute<FVector2D>::FGetter::CreateSP(this, &SDMXPixelMappingDesignerView::GetExtensionPosition, Handle)))
+				.Size(TAttribute<FVector2D>::Create(TAttribute<FVector2D>::FGetter::CreateSP(this, &SDMXPixelMappingDesignerView::GetExtensionSize, Handle)))
+				[
+					Handle.ToSharedRef()
+				];
 		}
 	}
 }
@@ -873,7 +851,7 @@ EVisibility SDMXPixelMappingDesignerView::GetZoomPanVisibility() const
 {
 	if (UDMXPixelMappingRendererComponent* RendererComponent = CachedRendererComponent.Get())
 	{
-		if (RendererComponent->GetRendererInputTexture())
+		if (RendererComponent->GetRenderedInputTexture())
 		{
 			return EVisibility::Visible;
 		}
@@ -1076,14 +1054,6 @@ bool SDMXPixelMappingDesignerView::GetGraphSpaceCursorPosition(FVector2D& OutGra
 	return false;
 }
 
-void SDMXPixelMappingDesignerView::GetPreviewAreaAndSize(FVector2D& Area, FVector2D& Size) const
-{
-	check(SourceTextureViewport.IsValid());
-
-	Area = FVector2D(SourceTextureViewport->GetPreviewAreaWidth().Get(), SourceTextureViewport->GetPreviewAreaHeight().Get());
-	Size = Area;
-}
-
 FGeometry SDMXPixelMappingDesignerView::GetDesignerGeometry() const
 {
 	return PreviewHitTestRoot->GetTickSpaceGeometry();
diff --git a/Engine/Plugins/VirtualProduction/DMX/DMXPixelMapping/Source/DMXPixelMappingEditor/Private/Views/SDMXPixelMappingDesignerView.h b/Engine/Plugins/VirtualProduction/DMX/DMXPixelMapping/Source/DMXPixelMappingEditor/Private/Views/SDMXPixelMappingDesignerView.h
--- a/Engine/Plugins/VirtualProduction/DMX/DMXPixelMapping/Source/DMXPixelMappingEditor/Private/Views/SDMXPixelMappingDesignerView.h
+++ b/Engine/Plugins/VirtualProduction/DMX/DMXPixelMapping/Source/DMXPixelMappingEditor/Private/Views/SDMXPixelMappingDesignerView.h
@@ -16,6 +16,7 @@ class IDMXPixelMappingOutputComponentWidgetInterface;
 class SDMXPixelMappingRuler;
 class SDMXPixelMappingSourceTextureViewport;
 class SDMXPixelMappingTransformHandle;
+class SDMXPixelMappingZoomPan;
 class UDMXPixelMappingBaseComponent;
 class UDMXPixelMappingOutputComponent;
 class UDMXPixelMappingRendererComponent;
@@ -31,7 +32,7 @@ class SOverlay;
 
 class SDMXPixelMappingDesignerView
 	: public SDMXPixelMappingSurface
-	, public FEditorUndoClient
+	, public FSelfRegisteringEditorUndoClient
 {
 private:
 	struct FComponentHitResult
@@ -47,13 +48,11 @@ private:
 		{}
 	};
 public:
-
-	SLATE_BEGIN_ARGS(SDMXPixelMappingDesignerView) { }
+	SLATE_BEGIN_ARGS(SDMXPixelMappingDesignerView) 
+	{}
+	
 	SLATE_END_ARGS()
 
-	/** Destructor */
-	virtual ~SDMXPixelMappingDesignerView();
-
 	/**
 	 * Constructs the widget.
 	 *
@@ -61,14 +60,8 @@ public:
 	 */
 	void Construct(const FArguments& InArgs, const TSharedPtr<FDMXPixelMappingToolkit>& InToolkit);
 
-	/** The width of the preview screen for the UI */
-	FOptionalSize GetPreviewAreaWidth() const;
-
-	/** The height of the preview screen for the UI */
-	FOptionalSize GetPreviewAreaHeight() const;
-
-	/** Gets the scale of the preview screen */
-	float GetPreviewScale() const;
+	/** Returns the geometry of the graph */
+	const FGeometry& GetGraphTickSpaceGeometry() const;
 
 protected:
 	// Begin SWidget interface
@@ -172,8 +165,6 @@ private:
 	/** Returns the cursor position in graph space */
 	bool GetGraphSpaceCursorPosition(FVector2D& OutGraphSpaceCursorPosition) const;
 
-	void GetPreviewAreaAndSize(FVector2D& Area, FVector2D& Size) const;
-
 	FGeometry GetDesignerGeometry() const;
 
 	FGeometry MakeGeometryWindowLocal(const FGeometry& WidgetGeometry) const;
@@ -210,6 +201,8 @@ private:
 
 	TWeakObjectPtr<UDMXPixelMappingRendererComponent> CachedRendererComponent;
 
+	TSharedPtr<SDMXPixelMappingZoomPan> ZoomPan;
+
 	/** The ruler bar at the top of the designer. */
 	TSharedPtr<SDMXPixelMappingRuler> TopRuler;
 
diff --git a/Engine/Plugins/VirtualProduction/DMX/DMXPixelMapping/Source/DMXPixelMappingEditor/Private/Views/SDMXPixelMappingInputSourceView.cpp b/Engine/Plugins/VirtualProduction/DMX/DMXPixelMapping/Source/DMXPixelMappingEditor/Private/Views/SDMXPixelMappingInputSourceView.cpp
--- a/Engine/Plugins/VirtualProduction/DMX/DMXPixelMapping/Source/DMXPixelMappingEditor/Private/Views/SDMXPixelMappingInputSourceView.cpp
+++ b/Engine/Plugins/VirtualProduction/DMX/DMXPixelMapping/Source/DMXPixelMappingEditor/Private/Views/SDMXPixelMappingInputSourceView.cpp
@@ -14,6 +14,7 @@
 #include "PropertyEditorDelegates.h"
 #include "PropertyEditorModule.h"
 #include "Toolkits/DMXPixelMappingToolkit.h"
+#include "Views/SDMXPixelMappingDesignerView.h"
 #include "Widgets/SBoxPanel.h"
 #include "Widgets/Images/SImage.h"
 #include "Widgets/Input/SButton.h"
@@ -241,12 +242,16 @@ FReply SDMXPixelMappingInputSourceView::OnAddButtonClicked()
 	{
 		Toolkit->AddRenderer();
 		
-		if (UDMXPixelMappingRendererComponent* RendererComponentToSelect = Toolkit->GetActiveRendererComponent())
+		UDMXPixelMappingRendererComponent* RendererComponentToSelect = Toolkit->GetActiveRendererComponent();
+		if (RendererComponentToSelect)
 		{
 			const FDMXPixelMappingComponentReference ComponentRefToSelect(Toolkit, RendererComponentToSelect);
 			const TSet<FDMXPixelMappingComponentReference> NewSelection{ ComponentRefToSelect };
 			Toolkit->SelectComponents(NewSelection);
-			Toolkit->ZoomToFit();
+
+			const TSharedRef<SDMXPixelMappingDesignerView> DesignerView = Toolkit->GetOrCreateDesignerView();
+			constexpr bool bInstantZoom = false;
+			DesignerView->ZoomToFit(bInstantZoom);
 		}
 
 		RequestRefresh();
diff --git a/Engine/Plugins/VirtualProduction/DMX/DMXPixelMapping/Source/DMXPixelMappingEditor/Private/Views/SDMXPixelMappingPreviewView.cpp b/Engine/Plugins/VirtualProduction/DMX/DMXPixelMapping/Source/DMXPixelMappingEditor/Private/Views/SDMXPixelMappingPreviewView.cpp
--- a/Engine/Plugins/VirtualProduction/DMX/DMXPixelMapping/Source/DMXPixelMappingEditor/Private/Views/SDMXPixelMappingPreviewView.cpp
+++ b/Engine/Plugins/VirtualProduction/DMX/DMXPixelMapping/Source/DMXPixelMappingEditor/Private/Views/SDMXPixelMappingPreviewView.cpp
@@ -300,7 +300,7 @@ EVisibility SDMXPixelMappingPreviewView::IsZoomPanVisible() const
 	{
 		if (UDMXPixelMappingRendererComponent* RendererComponent = Toolkit->GetActiveRendererComponent())
 		{
-			if (RendererComponent->GetRendererInputTexture())
+			if (RendererComponent->GetRenderedInputTexture())
 			{
 				return EVisibility::Visible;
 			}
diff --git a/Engine/Plugins/VirtualProduction/DMX/DMXPixelMapping/Source/DMXPixelMappingEditor/Private/Widgets/SDMXPixelMappingSourceTextureViewport.cpp b/Engine/Plugins/VirtualProduction/DMX/DMXPixelMapping/Source/DMXPixelMappingEditor/Private/Widgets/SDMXPixelMappingSourceTextureViewport.cpp
--- a/Engine/Plugins/VirtualProduction/DMX/DMXPixelMapping/Source/DMXPixelMappingEditor/Private/Widgets/SDMXPixelMappingSourceTextureViewport.cpp
+++ b/Engine/Plugins/VirtualProduction/DMX/DMXPixelMapping/Source/DMXPixelMappingEditor/Private/Widgets/SDMXPixelMappingSourceTextureViewport.cpp
@@ -1,107 +1,96 @@
 // Copyright Epic Games, Inc. All Rights Reserved.
 
 #include "Widgets/SDMXPixelMappingSourceTextureViewport.h"
-#include "TextureResource.h"
-#include "Viewports/DMXPixelMappingSourceTextureViewportClient.h"
-#include "Toolkits/DMXPixelMappingToolkit.h"
-#include "Components/DMXPixelMappingRendererComponent.h"
-#include "Viewports/DMXPixelMappingSceneViewport.h"
 
+#include "Components/DMXPixelMappingRendererComponent.h"
 #include "Engine/Texture.h"
 #include "Framework/Application/SlateApplication.h"
-#include "Widgets/SViewport.h"
+#include "Toolkits/DMXPixelMappingToolkit.h"
 #include "Widgets/Layout/SBox.h"
+#include "Widgets/SViewport.h"
+#include "Viewports/DMXPixelMappingSceneViewport.h"
+#include "Viewports/DMXPixelMappingSourceTextureViewportClient.h"
+#include "Views/SDMXPixelMappingDesignerView.h"
+
 
 void SDMXPixelMappingSourceTextureViewport::Construct(const FArguments& InArgs, const TSharedPtr<FDMXPixelMappingToolkit>& InToolkit)
 {
-	bIsRenderingEnabled = true;
-	ToolkitWeakPtr = InToolkit;
+	WeakToolkit = InToolkit;
 
 	TSharedPtr<SHorizontalBox> HorizontalBox;
 
 	ChildSlot
-	[
-		SNew(SBox)
-		.WidthOverride(this, &SDMXPixelMappingSourceTextureViewport::GetPreviewAreaWidth)
-		.HeightOverride(this, &SDMXPixelMappingSourceTextureViewport::GetPreviewAreaHeight)
+		[
+			SNew(SBox)
+			.Padding(this, &SDMXPixelMappingSourceTextureViewport::GetPaddingGraphSpace)
+			.WidthOverride(this, &SDMXPixelMappingSourceTextureViewport::GetWidthGraphSpace)
+			.HeightOverride(this, &SDMXPixelMappingSourceTextureViewport::GetHeightGraphSpace)
 			[
 				SAssignNew(ViewportWidget, SViewport)
-					.EnableGammaCorrection(false)
-					.IsEnabled(FSlateApplication::Get().GetNormalExecutionAttribute())
-					.ShowEffectWhenDisabled(false)
-					.EnableBlending(true)
+				.EnableGammaCorrection(false)
+				.IsEnabled(FSlateApplication::Get().GetNormalExecutionAttribute())
+				.ShowEffectWhenDisabled(false)
+				.EnableBlending(true)
 			]
-	];
+		];
 
 	ViewportClient = MakeShared<FDMXPixelMappingSourceTextureViewportClient>(InToolkit, SharedThis(this));
-
 	Viewport = MakeShared<FDMXPixelMappingSceneViewport>(ViewportClient.Get(), ViewportWidget);
 
-	// The viewport widget needs an interface so it knows what should render
 	ViewportWidget->SetViewportInterface(Viewport.ToSharedRef());
 }
 
-void SDMXPixelMappingSourceTextureViewport::EnableRendering()
-{
-	bIsRenderingEnabled = true;
-}
-
-void SDMXPixelMappingSourceTextureViewport::DisableRendering()
+void SDMXPixelMappingSourceTextureViewport::Tick(const FGeometry& AllottedGeometry, const double InCurrentTime, const float InDeltaTime)
 {
-	bIsRenderingEnabled = false;
+	Viewport->Invalidate();
 }
 
-void SDMXPixelMappingSourceTextureViewport::Tick(const FGeometry& AllottedGeometry, const double InCurrentTime, const float InDeltaTime)
+FOptionalSize SDMXPixelMappingSourceTextureViewport::GetWidthGraphSpace() const
 {
-	if (bIsRenderingEnabled)
+	if (ViewportClient->DrawOnlyVisibleRect())
 	{
-		Viewport->Invalidate();
+		return ViewportClient->GetVisibleTextureBoxGraphSpace().Max.X;
 	}
+	else if (UTexture* InputTexture = GetInputTexture())
+	{
+		return InputTexture->GetSurfaceWidth();
+	}
+	return FOptionalSize();
 }
 
-FOptionalSize SDMXPixelMappingSourceTextureViewport::GetPreviewAreaWidth() const
+FOptionalSize SDMXPixelMappingSourceTextureViewport::GetHeightGraphSpace() const
 {
-	float Size = 1.f;
-
-	if (const FTextureResource* Resource = GetInputTextureResource())
+	if (ViewportClient->DrawOnlyVisibleRect())
 	{
-		Size = Resource->GetSizeX();
+		return ViewportClient->GetVisibleTextureBoxGraphSpace().Max.Y;
 	}
-
-	return Size;
+	else if (UTexture* InputTexture = GetInputTexture())
+	{
+		return InputTexture->GetSurfaceHeight();
+	}
+	return FOptionalSize();
 }
 
-FOptionalSize SDMXPixelMappingSourceTextureViewport::GetPreviewAreaHeight() const
+FMargin SDMXPixelMappingSourceTextureViewport::GetPaddingGraphSpace() const
 {
-	float Size = 1.f;
-
-	if (const FTextureResource* Resource = GetInputTextureResource())
+	if (ViewportClient->DrawOnlyVisibleRect())
 	{
-		Size = Resource->GetSizeY();
+		return FMargin(ViewportClient->GetVisibleTextureBoxGraphSpace().Min.X, ViewportClient->GetVisibleTextureBoxGraphSpace().Min.Y, 0.0, 0.0);
+	}
+	else
+	{
+		return 0.0;
 	}
-
-	return Size;
 }
 
 UTexture* SDMXPixelMappingSourceTextureViewport::GetInputTexture() const
 {
-	if (TSharedPtr<FDMXPixelMappingToolkit> Toolkit = ToolkitWeakPtr.Pin())
+	if (TSharedPtr<FDMXPixelMappingToolkit> Toolkit = WeakToolkit.Pin())
 	{
 		if (UDMXPixelMappingRendererComponent* RendererComponent = Toolkit->GetActiveRendererComponent())
 		{
-			return RendererComponent->GetRendererInputTexture();
+			return RendererComponent->GetRenderedInputTexture();
 		}
 	}
-
-	return nullptr;
-}
-
-const FTextureResource* SDMXPixelMappingSourceTextureViewport::GetInputTextureResource() const
-{
-	if (const UTexture* Texture = GetInputTexture())
-	{
-		return Texture->GetResource();
-	}
-
 	return nullptr;
 }
diff --git a/Engine/Plugins/VirtualProduction/DMX/DMXPixelMapping/Source/DMXPixelMappingEditor/Private/Widgets/SDMXPixelMappingSourceTextureViewport.h b/Engine/Plugins/VirtualProduction/DMX/DMXPixelMapping/Source/DMXPixelMappingEditor/Private/Widgets/SDMXPixelMappingSourceTextureViewport.h
--- a/Engine/Plugins/VirtualProduction/DMX/DMXPixelMapping/Source/DMXPixelMappingEditor/Private/Widgets/SDMXPixelMappingSourceTextureViewport.h
+++ b/Engine/Plugins/VirtualProduction/DMX/DMXPixelMapping/Source/DMXPixelMappingEditor/Private/Widgets/SDMXPixelMappingSourceTextureViewport.h
@@ -16,8 +16,9 @@ class SDMXPixelMappingSourceTextureViewport
 	: public SCompoundWidget
 {
 public:
-
-	SLATE_BEGIN_ARGS(SDMXPixelMappingSourceTextureViewport) { }
+	SLATE_BEGIN_ARGS(SDMXPixelMappingSourceTextureViewport) 
+	{}
+	
 	SLATE_END_ARGS()
 
 public:
@@ -28,12 +29,6 @@ public:
 	 */
 	void Construct(const FArguments& InArgs, const TSharedPtr<FDMXPixelMappingToolkit>& InToolkit);
 
-	/** Enable viewport rendering */
-	void EnableRendering();
-
-	/** Disable viewport rendering */
-	void DisableRendering();
-
 	TSharedPtr<FSceneViewport> GetViewport() const { return Viewport; }
 
 	TSharedPtr<SViewport> GetViewportWidget() const { return ViewportWidget; }
@@ -42,17 +37,21 @@ public:
 	virtual void Tick(const FGeometry& AllottedGeometry, const double InCurrentTime, const float InDeltaTime) override;
 	// End of SWidget interface
 
-	FOptionalSize GetPreviewAreaWidth() const;
+private:
+	/** Returns the texture width, in graph space  */
+	FOptionalSize GetWidthGraphSpace() const;
+
+	/** Returns the texture height, in graph space  */
+	FOptionalSize GetHeightGraphSpace() const;
 
-	FOptionalSize GetPreviewAreaHeight() const;
+	/** Returns the padding of content, in graph space */
+	FMargin GetPaddingGraphSpace() const;
 
+	/** Returns the input texture currently in use */
 	UTexture* GetInputTexture() const;
 
-	const FTextureResource* GetInputTextureResource() const;
-
-private:
 	/** Weak pointer to the editor toolkit */
-	TWeakPtr<FDMXPixelMappingToolkit> ToolkitWeakPtr;
+	TWeakPtr<FDMXPixelMappingToolkit> WeakToolkit;
 
 	/** viewport client. */
 	TSharedPtr<FDMXPixelMappingSourceTextureViewportClient> ViewportClient;
@@ -62,7 +61,4 @@ private:
 
 	/** Viewport widget. */
 	TSharedPtr<SViewport> ViewportWidget;
-
-	/** Is rendering currently enabled? (disabled when reimporting a texture) */
-	bool bIsRenderingEnabled;
 };
diff --git a/Engine/Plugins/VirtualProduction/DMX/DMXPixelMapping/Source/DMXPixelMappingEditor/Private/Widgets/SDMXPixelMappingSurface.cpp b/Engine/Plugins/VirtualProduction/DMX/DMXPixelMapping/Source/DMXPixelMappingEditor/Private/Widgets/SDMXPixelMappingSurface.cpp
--- a/Engine/Plugins/VirtualProduction/DMX/DMXPixelMapping/Source/DMXPixelMappingEditor/Private/Widgets/SDMXPixelMappingSurface.cpp
+++ b/Engine/Plugins/VirtualProduction/DMX/DMXPixelMapping/Source/DMXPixelMappingEditor/Private/Widgets/SDMXPixelMappingSurface.cpp
@@ -30,54 +30,58 @@ struct FFixedZoomLevelsContainerDesignSurface : public FZoomLevelsContainer
 {
 	FFixedZoomLevelsContainerDesignSurface()
 	{
-		ZoomLevels.Add(FZoomLevelEntry(0.025f, FText::FromString("-15"), EGraphRenderingLOD::LowestDetail));
-		ZoomLevels.Add(FZoomLevelEntry(0.050f, FText::FromString("-14"), EGraphRenderingLOD::LowestDetail));
-		ZoomLevels.Add(FZoomLevelEntry(0.075f, FText::FromString("-13"), EGraphRenderingLOD::LowestDetail));
-		ZoomLevels.Add(FZoomLevelEntry(0.100f, FText::FromString("-12"), EGraphRenderingLOD::LowestDetail));
-		ZoomLevels.Add(FZoomLevelEntry(0.125f, FText::FromString("-11"), EGraphRenderingLOD::LowestDetail));
-		ZoomLevels.Add(FZoomLevelEntry(0.150f, FText::FromString("-10"), EGraphRenderingLOD::LowestDetail));
-		ZoomLevels.Add(FZoomLevelEntry(0.175f, FText::FromString("-9"), EGraphRenderingLOD::LowestDetail));
-		ZoomLevels.Add(FZoomLevelEntry(0.200f, FText::FromString("-8"), EGraphRenderingLOD::LowestDetail));
-		ZoomLevels.Add(FZoomLevelEntry(0.225f, FText::FromString("-7"), EGraphRenderingLOD::LowDetail));
-		ZoomLevels.Add(FZoomLevelEntry(0.250f, FText::FromString("-6"), EGraphRenderingLOD::LowDetail));
-		ZoomLevels.Add(FZoomLevelEntry(0.375f, FText::FromString("-5"), EGraphRenderingLOD::MediumDetail));
-		ZoomLevels.Add(FZoomLevelEntry(0.500f, FText::FromString("-4"), EGraphRenderingLOD::MediumDetail));
-		ZoomLevels.Add(FZoomLevelEntry(0.675f, FText::FromString("-3"), EGraphRenderingLOD::MediumDetail));
-		ZoomLevels.Add(FZoomLevelEntry(0.750f, FText::FromString("-2"), EGraphRenderingLOD::DefaultDetail));
-		ZoomLevels.Add(FZoomLevelEntry(0.875f, FText::FromString("-1"), EGraphRenderingLOD::DefaultDetail));
+		ZoomLevels.Add(FZoomLevelEntry(0.025f, FText::FromString("1:40"), EGraphRenderingLOD::LowestDetail));
+		ZoomLevels.Add(FZoomLevelEntry(0.050f, FText::FromString("1:20"), EGraphRenderingLOD::LowestDetail));
+		ZoomLevels.Add(FZoomLevelEntry(0.066666666666f, FText::FromString("1:15"), EGraphRenderingLOD::LowestDetail));
+		ZoomLevels.Add(FZoomLevelEntry(0.100f, FText::FromString("1:10"), EGraphRenderingLOD::LowestDetail));
+		ZoomLevels.Add(FZoomLevelEntry(0.111111111111f, FText::FromString("1:9"), EGraphRenderingLOD::MediumDetail));
+		ZoomLevels.Add(FZoomLevelEntry(0.125f, FText::FromString("1:8"), EGraphRenderingLOD::MediumDetail));
+		ZoomLevels.Add(FZoomLevelEntry(0.142857142857f, FText::FromString("1:7"), EGraphRenderingLOD::MediumDetail));
+		ZoomLevels.Add(FZoomLevelEntry(0.166666666666f, FText::FromString("1:6"), EGraphRenderingLOD::MediumDetail));
+		ZoomLevels.Add(FZoomLevelEntry(0.2f, FText::FromString("1:5"), EGraphRenderingLOD::MediumDetail));
+		ZoomLevels.Add(FZoomLevelEntry(0.25f, FText::FromString("1:4"), EGraphRenderingLOD::MediumDetail));
+		ZoomLevels.Add(FZoomLevelEntry(0.333333333333f, FText::FromString("1:3"), EGraphRenderingLOD::MediumDetail));
+		ZoomLevels.Add(FZoomLevelEntry(0.5f, FText::FromString("1:2"), EGraphRenderingLOD::MediumDetail));
 		ZoomLevels.Add(FZoomLevelEntry(1.000f, FText::FromString("1:1"), EGraphRenderingLOD::DefaultDetail));
-		ZoomLevels.Add(FZoomLevelEntry(1.250f, FText::FromString("+1"), EGraphRenderingLOD::DefaultDetail));
-		ZoomLevels.Add(FZoomLevelEntry(1.500f, FText::FromString("+2"), EGraphRenderingLOD::DefaultDetail));
-		ZoomLevels.Add(FZoomLevelEntry(1.750f, FText::FromString("+3"), EGraphRenderingLOD::FullyZoomedIn));
-		ZoomLevels.Add(FZoomLevelEntry(2.000f, FText::FromString("+4"), EGraphRenderingLOD::FullyZoomedIn));
-		ZoomLevels.Add(FZoomLevelEntry(2.250f, FText::FromString("+5"), EGraphRenderingLOD::FullyZoomedIn));
-		ZoomLevels.Add(FZoomLevelEntry(2.500f, FText::FromString("+6"), EGraphRenderingLOD::FullyZoomedIn));
-		ZoomLevels.Add(FZoomLevelEntry(2.750f, FText::FromString("+7"), EGraphRenderingLOD::FullyZoomedIn));
-		ZoomLevels.Add(FZoomLevelEntry(3.000f, FText::FromString("+8"), EGraphRenderingLOD::FullyZoomedIn));
-		ZoomLevels.Add(FZoomLevelEntry(3.250f, FText::FromString("+9"), EGraphRenderingLOD::FullyZoomedIn));
-		ZoomLevels.Add(FZoomLevelEntry(3.500f, FText::FromString("+10"), EGraphRenderingLOD::FullyZoomedIn));
-		ZoomLevels.Add(FZoomLevelEntry(4.000f, FText::FromString("+11"), EGraphRenderingLOD::FullyZoomedIn));
-		ZoomLevels.Add(FZoomLevelEntry(5.000f, FText::FromString("+12"), EGraphRenderingLOD::FullyZoomedIn));
-		ZoomLevels.Add(FZoomLevelEntry(6.000f, FText::FromString("+13"), EGraphRenderingLOD::FullyZoomedIn));
-		ZoomLevels.Add(FZoomLevelEntry(7.000f, FText::FromString("+14"), EGraphRenderingLOD::FullyZoomedIn));
-		ZoomLevels.Add(FZoomLevelEntry(8.000f, FText::FromString("+15"), EGraphRenderingLOD::FullyZoomedIn));
-		ZoomLevels.Add(FZoomLevelEntry(9.000f, FText::FromString("+16"), EGraphRenderingLOD::FullyZoomedIn));
-		ZoomLevels.Add(FZoomLevelEntry(10.000f, FText::FromString("+17"), EGraphRenderingLOD::FullyZoomedIn));
-		ZoomLevels.Add(FZoomLevelEntry(11.000f, FText::FromString("+18"), EGraphRenderingLOD::FullyZoomedIn));
-		ZoomLevels.Add(FZoomLevelEntry(12.000f, FText::FromString("+19"), EGraphRenderingLOD::FullyZoomedIn));
-				ZoomLevels.Add(FZoomLevelEntry(13.000f, FText::FromString("+20"), EGraphRenderingLOD::FullyZoomedIn));
-		ZoomLevels.Add(FZoomLevelEntry(14.000f, FText::FromString("+21"), EGraphRenderingLOD::FullyZoomedIn));
-		ZoomLevels.Add(FZoomLevelEntry(15.000f, FText::FromString("+22"), EGraphRenderingLOD::FullyZoomedIn));
-		ZoomLevels.Add(FZoomLevelEntry(16.000f, FText::FromString("+23"), EGraphRenderingLOD::FullyZoomedIn));
-		ZoomLevels.Add(FZoomLevelEntry(17.000f, FText::FromString("+24"), EGraphRenderingLOD::FullyZoomedIn));
-		ZoomLevels.Add(FZoomLevelEntry(18.000f, FText::FromString("+25"), EGraphRenderingLOD::FullyZoomedIn));
-		ZoomLevels.Add(FZoomLevelEntry(19.000f, FText::FromString("+26"), EGraphRenderingLOD::FullyZoomedIn));
-		ZoomLevels.Add(FZoomLevelEntry(20.000f, FText::FromString("+27"), EGraphRenderingLOD::FullyZoomedIn));
-		ZoomLevels.Add(FZoomLevelEntry(21.000f, FText::FromString("+28"), EGraphRenderingLOD::FullyZoomedIn));
-		ZoomLevels.Add(FZoomLevelEntry(22.000f, FText::FromString("+29"), EGraphRenderingLOD::FullyZoomedIn));
-		ZoomLevels.Add(FZoomLevelEntry(23.000f, FText::FromString("+30"), EGraphRenderingLOD::FullyZoomedIn));
-	}
+		ZoomLevels.Add(FZoomLevelEntry(1.250f, FText::FromString("1.25:1"), EGraphRenderingLOD::DefaultDetail));
+		ZoomLevels.Add(FZoomLevelEntry(1.500f, FText::FromString("1.5:1"), EGraphRenderingLOD::DefaultDetail));
+		ZoomLevels.Add(FZoomLevelEntry(1.750f, FText::FromString("1.75:1"), EGraphRenderingLOD::FullyZoomedIn));
+		ZoomLevels.Add(FZoomLevelEntry(2.000f, FText::FromString("2:1"), EGraphRenderingLOD::FullyZoomedIn));
+		ZoomLevels.Add(FZoomLevelEntry(2.250f, FText::FromString("2.25:1"), EGraphRenderingLOD::FullyZoomedIn));
+		ZoomLevels.Add(FZoomLevelEntry(2.500f, FText::FromString("2.5:1"), EGraphRenderingLOD::FullyZoomedIn));
+		ZoomLevels.Add(FZoomLevelEntry(2.750f, FText::FromString("2.75:1"), EGraphRenderingLOD::FullyZoomedIn));
+		ZoomLevels.Add(FZoomLevelEntry(3.000f, FText::FromString("3:1"), EGraphRenderingLOD::FullyZoomedIn));
+		ZoomLevels.Add(FZoomLevelEntry(3.250f, FText::FromString("3.25:1"), EGraphRenderingLOD::FullyZoomedIn));
+		ZoomLevels.Add(FZoomLevelEntry(3.500f, FText::FromString("3.5:1"), EGraphRenderingLOD::FullyZoomedIn));
+		ZoomLevels.Add(FZoomLevelEntry(4.000f, FText::FromString("4:1"), EGraphRenderingLOD::FullyZoomedIn));
+		ZoomLevels.Add(FZoomLevelEntry(5.000f, FText::FromString("5:1"), EGraphRenderingLOD::FullyZoomedIn));
+		ZoomLevels.Add(FZoomLevelEntry(6.000f, FText::FromString("6:1"), EGraphRenderingLOD::FullyZoomedIn));
+		ZoomLevels.Add(FZoomLevelEntry(7.000f, FText::FromString("7:1"), EGraphRenderingLOD::FullyZoomedIn));
+		ZoomLevels.Add(FZoomLevelEntry(8.000f, FText::FromString("8:1"), EGraphRenderingLOD::FullyZoomedIn));
+		ZoomLevels.Add(FZoomLevelEntry(9.000f, FText::FromString("9:1"), EGraphRenderingLOD::FullyZoomedIn));
+		ZoomLevels.Add(FZoomLevelEntry(10.000f, FText::FromString("10:1"), EGraphRenderingLOD::FullyZoomedIn));
+		ZoomLevels.Add(FZoomLevelEntry(11.000f, FText::FromString("11:1"), EGraphRenderingLOD::FullyZoomedIn));
+		ZoomLevels.Add(FZoomLevelEntry(12.000f, FText::FromString("12:1"), EGraphRenderingLOD::FullyZoomedIn));
+		ZoomLevels.Add(FZoomLevelEntry(13.000f, FText::FromString("13:1"), EGraphRenderingLOD::FullyZoomedIn));
+		ZoomLevels.Add(FZoomLevelEntry(14.000f, FText::FromString("14:1"), EGraphRenderingLOD::FullyZoomedIn));
+		ZoomLevels.Add(FZoomLevelEntry(15.000f, FText::FromString("15:1"), EGraphRenderingLOD::FullyZoomedIn));
+		ZoomLevels.Add(FZoomLevelEntry(16.000f, FText::FromString("16:1"), EGraphRenderingLOD::FullyZoomedIn));
+		ZoomLevels.Add(FZoomLevelEntry(17.000f, FText::FromString("17:1"), EGraphRenderingLOD::FullyZoomedIn));
+		ZoomLevels.Add(FZoomLevelEntry(18.000f, FText::FromString("18:1"), EGraphRenderingLOD::FullyZoomedIn));
+		ZoomLevels.Add(FZoomLevelEntry(19.000f, FText::FromString("19:1"), EGraphRenderingLOD::FullyZoomedIn));
+		ZoomLevels.Add(FZoomLevelEntry(20.000f, FText::FromString("20:1"), EGraphRenderingLOD::FullyZoomedIn));
+		ZoomLevels.Add(FZoomLevelEntry(21.000f, FText::FromString("20:1"), EGraphRenderingLOD::FullyZoomedIn));
+		ZoomLevels.Add(FZoomLevelEntry(22.000f, FText::FromString("22:1"), EGraphRenderingLOD::FullyZoomedIn));
+		ZoomLevels.Add(FZoomLevelEntry(23.000f, FText::FromString("23:1"), EGraphRenderingLOD::FullyZoomedIn));
+		ZoomLevels.Add(FZoomLevelEntry(24.000f, FText::FromString("24:1"), EGraphRenderingLOD::FullyZoomedIn));
+		ZoomLevels.Add(FZoomLevelEntry(25.000f, FText::FromString("25:1"), EGraphRenderingLOD::FullyZoomedIn));
+		ZoomLevels.Add(FZoomLevelEntry(30.000f, FText::FromString("30:1"), EGraphRenderingLOD::FullyZoomedIn));
+		ZoomLevels.Add(FZoomLevelEntry(40.000f, FText::FromString("40:1"), EGraphRenderingLOD::FullyZoomedIn));
+		ZoomLevels.Add(FZoomLevelEntry(50.000f, FText::FromString("50:1"), EGraphRenderingLOD::FullyZoomedIn));
+		ZoomLevels.Add(FZoomLevelEntry(100.000f, FText::FromString("100:1"), EGraphRenderingLOD::FullyZoomedIn));
 
+	}
+	
 	float GetZoomAmount(int32 InZoomLevel) const override
 	{
 		checkSlow(ZoomLevels.IsValidIndex(InZoomLevel));
@@ -372,9 +376,10 @@ FReply SDMXPixelMappingSurface::OnTouchEnded(const FGeometry& MyGeometry, const
 	return FReply::Unhandled();
 }
 
-inline float FancyMod(float Value, float Size)
+
+FVector2D SDMXPixelMappingSurface::GetViewOffset() const
 {
-	return ((Value >= 0) ? 0.0f : Size) + FMath::Fmod(Value, Size);
+	return ViewOffset;
 }
 
 float SDMXPixelMappingSurface::GetZoomAmount() const
@@ -551,11 +556,6 @@ FSlateRect SDMXPixelMappingSurface::ComputeAreaBounds() const
 	return Bounds;
 }
 
-FVector2D SDMXPixelMappingSurface::GetViewOffset() const
-{
-	return ViewOffset;
-}
-
 FVector2D SDMXPixelMappingSurface::GraphCoordToPanelCoord(const FVector2D& GraphSpaceCoordinate) const
 {
 	return (GraphSpaceCoordinate - GetViewOffset()) * GetZoomAmount();
diff --git a/Engine/Plugins/VirtualProduction/DMX/DMXPixelMapping/Source/DMXPixelMappingEditor/Private/Widgets/SDMXPixelMappingSurface.h b/Engine/Plugins/VirtualProduction/DMX/DMXPixelMapping/Source/DMXPixelMappingEditor/Private/Widgets/SDMXPixelMappingSurface.h
--- a/Engine/Plugins/VirtualProduction/DMX/DMXPixelMapping/Source/DMXPixelMappingEditor/Private/Widgets/SDMXPixelMappingSurface.h
+++ b/Engine/Plugins/VirtualProduction/DMX/DMXPixelMapping/Source/DMXPixelMappingEditor/Private/Widgets/SDMXPixelMappingSurface.h
@@ -45,6 +45,9 @@ public:
 	virtual bool SupportsKeyboardFocus() const override { return true; }
 	// End of Swidget interface
 
+	/** Returns the view offset */
+	FVector2D GetViewOffset() const;
+
 	/** Gets the current zoom factor. */
 	float GetZoomAmount() const;
 
@@ -70,8 +73,6 @@ protected:
 	FText GetZoomText() const;
 	FSlateColor GetZoomTextColorAndOpacity() const;
 
-	FVector2D GetViewOffset() const;
-
 	FSlateRect ComputeSensibleBounds() const;
 
 	FVector2D GraphCoordToPanelCoord(const FVector2D& GraphSpaceCoordinate) const;
diff --git a/Engine/Plugins/VirtualProduction/DMX/DMXPixelMapping/Source/DMXPixelMappingEditor/Private/Widgets/SDMXPixelMappingTransformHandle.cpp b/Engine/Plugins/VirtualProduction/DMX/DMXPixelMapping/Source/DMXPixelMappingEditor/Private/Widgets/SDMXPixelMappingTransformHandle.cpp
--- a/Engine/Plugins/VirtualProduction/DMX/DMXPixelMapping/Source/DMXPixelMappingEditor/Private/Widgets/SDMXPixelMappingTransformHandle.cpp
+++ b/Engine/Plugins/VirtualProduction/DMX/DMXPixelMapping/Source/DMXPixelMappingEditor/Private/Widgets/SDMXPixelMappingTransformHandle.cpp
@@ -89,7 +89,7 @@ FReply SDMXPixelMappingTransformHandle::OnMouseButtonUp(const FGeometry& MyGeome
 				{
 					// Set the final size transacted
 					const FVector2D Delta = MouseEvent.GetScreenSpacePosition() - MouseDownPosition;
-					const FVector2D TranslateAmount = Delta * (1.0f / (DesignerViewWeakPtr.Pin()->GetPreviewScale() * MyGeometry.Scale));
+					const FVector2D TranslateAmount = Delta * (1.0f / (DesignerViewWeakPtr.Pin()->GetZoomAmount() * MyGeometry.Scale));
 
 					ResizedComponent->Modify();
 					Resize(ResizedComponent, DragDirection, TranslateAmount);
@@ -117,7 +117,7 @@ FReply SDMXPixelMappingTransformHandle::OnMouseMove(const FGeometry& MyGeometry,
 				if (UDMXPixelMappingOutputComponent* ResizedComponent = Cast<UDMXPixelMappingOutputComponent>(ComponentReference.GetComponent()))
 				{
 					const FVector2D Delta = MouseEvent.GetScreenSpacePosition() - MouseDownPosition;
-					const FVector2D TranslateAmount = Delta * (1.0f / (DesignerViewWeakPtr.Pin()->GetPreviewScale() * MyGeometry.Scale));
+					const FVector2D TranslateAmount = Delta * (1.0f / (DesignerViewWeakPtr.Pin()->GetZoomAmount() * MyGeometry.Scale));
 
 					Resize(ResizedComponent, DragDirection, TranslateAmount);
 				}
diff --git a/Engine/Plugins/VirtualProduction/DMX/DMXPixelMapping/Source/DMXPixelMappingRuntime/Private/Components/DMXPixelMappingFixtureGroupItemComponent.cpp b/Engine/Plugins/VirtualProduction/DMX/DMXPixelMapping/Source/DMXPixelMappingRuntime/Private/Components/DMXPixelMappingFixtureGroupItemComponent.cpp
--- a/Engine/Plugins/VirtualProduction/DMX/DMXPixelMapping/Source/DMXPixelMappingRuntime/Private/Components/DMXPixelMappingFixtureGroupItemComponent.cpp
+++ b/Engine/Plugins/VirtualProduction/DMX/DMXPixelMapping/Source/DMXPixelMappingRuntime/Private/Components/DMXPixelMappingFixtureGroupItemComponent.cpp
@@ -322,7 +322,7 @@ void UDMXPixelMappingFixtureGroupItemComponent::QueueDownsample()
 		return;
 	}
 
-	UTexture* InputTexture = RendererComponent->GetRendererInputTexture();
+	UTexture* InputTexture = RendererComponent->GetRenderedInputTexture();
 	if (!ensure(InputTexture))
 	{
 		return;
diff --git a/Engine/Plugins/VirtualProduction/DMX/DMXPixelMapping/Source/DMXPixelMappingRuntime/Private/Components/DMXPixelMappingMatrixCellComponent.cpp b/Engine/Plugins/VirtualProduction/DMX/DMXPixelMapping/Source/DMXPixelMappingRuntime/Private/Components/DMXPixelMappingMatrixCellComponent.cpp
--- a/Engine/Plugins/VirtualProduction/DMX/DMXPixelMapping/Source/DMXPixelMappingRuntime/Private/Components/DMXPixelMappingMatrixCellComponent.cpp
+++ b/Engine/Plugins/VirtualProduction/DMX/DMXPixelMapping/Source/DMXPixelMappingRuntime/Private/Components/DMXPixelMappingMatrixCellComponent.cpp
@@ -233,7 +233,7 @@ void UDMXPixelMappingMatrixCellComponent::QueueDownsample()
 		return;
 	}
 
-	UTexture* InputTexture = RendererComponent->GetRendererInputTexture();
+	UTexture* InputTexture = RendererComponent->GetRenderedInputTexture();
 	if (!ensure(InputTexture))
 	{
 		return;
diff --git a/Engine/Plugins/VirtualProduction/DMX/DMXPixelMapping/Source/DMXPixelMappingRuntime/Private/Components/DMXPixelMappingRendererComponent.cpp b/Engine/Plugins/VirtualProduction/DMX/DMXPixelMapping/Source/DMXPixelMappingRuntime/Private/Components/DMXPixelMappingRendererComponent.cpp
--- a/Engine/Plugins/VirtualProduction/DMX/DMXPixelMapping/Source/DMXPixelMappingRuntime/Private/Components/DMXPixelMappingRendererComponent.cpp
+++ b/Engine/Plugins/VirtualProduction/DMX/DMXPixelMapping/Source/DMXPixelMappingRuntime/Private/Components/DMXPixelMappingRendererComponent.cpp
@@ -131,7 +131,20 @@ void UDMXPixelMappingRendererComponent::PostEditChangeChainProperty(FPropertyCha
 		}
 	}
 
-	Initialize();
+	if (PropertyChangedChainEvent.ChangeType == EPropertyChangeType::Interactive)
+	{
+		// Apply interactive changes only once per frame
+		static uint64 Frame = GFrameCounter;
+		if (Frame != GFrameCounter)
+		{
+			Initialize();
+			Frame = GFrameCounter;
+		}
+	}
+	else
+	{
+		Initialize();
+	}
 }
 #endif // WITH_EDITOR
 
@@ -300,7 +313,7 @@ void UDMXPixelMappingRendererComponent::Render()
 
 	// 1. Render the input texture
 	RendererInputTexture();
-	UTexture* DownsampleInputTexture = GetRendererInputTexture();
+	UTexture* DownsampleInputTexture = GetRenderedInputTexture();
 	if (!DownsampleInputTexture)
 	{
 		return;
@@ -369,11 +382,6 @@ UWorld* UDMXPixelMappingRendererComponent::GetWorld() const
 
 void UDMXPixelMappingRendererComponent::Initialize()
 {
-	if (InputWidget)
-	{
-		UserWidget = CreateWidget(GetWorld(), InputWidget);
-	}
-
 	if (!PixelMappingRenderer.IsValid())
 	{
 		PixelMappingRenderer = IDMXPixelMappingRendererModule::Get().CreateRenderer();
@@ -391,6 +399,7 @@ void UDMXPixelMappingRendererComponent::Initialize()
 		break;
 
 	case(EDMXPixelMappingRendererType::UMG):
+		UserWidget = CreateWidget(GetWorld(), InputWidget);
 		PreprocessRenderer->SetInputUserWidget(UserWidget.Get());
 		break;
 
@@ -417,23 +426,9 @@ void UDMXPixelMappingRendererComponent::RendererInputTexture()
 	SCOPE_CYCLE_COUNTER(STAT_RenderInputTexture);
 
 	PreprocessRenderer->Render();
+}
 
-//#if WITH_EDITOR
-//	if (RendererType == EDMXPixelMappingRendererType::Texture)
-//	{
-//		if (InputTexture && InputTexture->GetResource())
-//		{
-//			ResizePreviewRenderTarget(InputTexture->GetResource()->GetSizeX(), InputTexture->GetResource()->GetSizeY());
-//		}
-//	}
-//	else
-//	{
-//		ResizePreviewRenderTarget(GetSize().X, GetSize().Y);
-//	}
-//#endif
-}
-
-UTexture* UDMXPixelMappingRendererComponent::GetRendererInputTexture() const
+UTexture* UDMXPixelMappingRendererComponent::GetRenderedInputTexture() const
 {
 	return PreprocessRenderer->GetRenderedTexture();
 }
diff --git a/Engine/Plugins/VirtualProduction/DMX/DMXPixelMapping/Source/DMXPixelMappingRuntime/Private/Components/DMXPixelMappingScreenComponent.cpp b/Engine/Plugins/VirtualProduction/DMX/DMXPixelMapping/Source/DMXPixelMappingRuntime/Private/Components/DMXPixelMappingScreenComponent.cpp
--- a/Engine/Plugins/VirtualProduction/DMX/DMXPixelMapping/Source/DMXPixelMappingRuntime/Private/Components/DMXPixelMappingScreenComponent.cpp
+++ b/Engine/Plugins/VirtualProduction/DMX/DMXPixelMapping/Source/DMXPixelMappingRuntime/Private/Components/DMXPixelMappingScreenComponent.cpp
@@ -383,7 +383,7 @@ void UDMXPixelMappingScreenComponent::QueueDownsample()
 		return;
 	}
 
-	UTexture* InputTexture = RendererComponent->GetRendererInputTexture();
+	UTexture* InputTexture = RendererComponent->GetRenderedInputTexture();
 	if (!ensure(InputTexture))
 	{
 		return;
diff --git a/Engine/Plugins/VirtualProduction/DMX/DMXPixelMapping/Source/DMXPixelMappingRuntime/Public/Components/DMXPixelMappingRendererComponent.h b/Engine/Plugins/VirtualProduction/DMX/DMXPixelMapping/Source/DMXPixelMappingRuntime/Public/Components/DMXPixelMappingRendererComponent.h
--- a/Engine/Plugins/VirtualProduction/DMX/DMXPixelMapping/Source/DMXPixelMappingRuntime/Public/Components/DMXPixelMappingRendererComponent.h
+++ b/Engine/Plugins/VirtualProduction/DMX/DMXPixelMapping/Source/DMXPixelMappingRuntime/Public/Components/DMXPixelMappingRendererComponent.h
@@ -60,8 +60,8 @@ public:
 	UTextureRenderTarget2D* GetPreviewRenderTarget();
 #endif
 
-	/** Get reference to the active input texture */
-	UTexture* GetRendererInputTexture() const;
+	/** Gets the rendered input texture, or nullptr if no input texture is currently rendered */
+	UTexture* GetRenderedInputTexture() const;
 
 	/**
 	 * Get pixel position in downsample buffer target based on pixel index
