DMX - Fix crash when removing Fixture Patches from DMX Library that are used in pixel mapping

diff --git a/Engine/Plugins/VirtualProduction/DMX/DMXPixelMapping/Source/DMXPixelMappingRuntime/Private/Components/DMXPixelMappingMatrixComponent.cpp b/Engine/Plugins/VirtualProduction/DMX/DMXPixelMapping/Source/DMXPixelMappingRuntime/Private/Components/DMXPixelMappingMatrixComponent.cpp
--- a/Engine/Plugins/VirtualProduction/DMX/DMXPixelMapping/Source/DMXPixelMappingRuntime/Private/Components/DMXPixelMappingMatrixComponent.cpp
+++ b/Engine/Plugins/VirtualProduction/DMX/DMXPixelMapping/Source/DMXPixelMappingRuntime/Private/Components/DMXPixelMappingMatrixComponent.cpp
@@ -474,6 +474,14 @@ void UDMXPixelMappingMatrixComponent::HandleMatrixChanged()
 {
 	TGuardValue<bool>(bIsUpdatingChildren, true);
 
+	UDMXPixelMapping* PixelMapping = GetPixelMapping();
+	UDMXPixelMappingRootComponent* RootComponent = PixelMapping ? PixelMapping->GetRootComponent() : nullptr;
+	if (!RootComponent)
+	{
+		// Component was removed from pixel mapping
+		return;
+	}
+
 	CoordinateGrid = 0;
 
 	// Remove all existing children and rebuild them anew
@@ -496,7 +504,7 @@ void UDMXPixelMappingMatrixComponent::HandleMatrixChanged()
 			for (const FDMXCell& Cell : MatrixCells)
 			{
 				TSharedPtr<FDMXPixelMappingComponentTemplate> ComponentTemplate = MakeShared<FDMXPixelMappingComponentTemplate>(UDMXPixelMappingMatrixCellComponent::StaticClass());
-				UDMXPixelMappingMatrixCellComponent* Component = ComponentTemplate->CreateComponent<UDMXPixelMappingMatrixCellComponent>(GetPixelMapping()->GetRootComponent());
+				UDMXPixelMappingMatrixCellComponent* Component = ComponentTemplate->CreateComponent<UDMXPixelMappingMatrixCellComponent>(RootComponent);
 
 				Component->CellID = Cell.CellID;
 				Component->SetCellCoordinate(Cell.Coordinate);
@@ -508,7 +516,7 @@ void UDMXPixelMappingMatrixComponent::HandleMatrixChanged()
 
 	HandleSizeChanged();
 
-	GetOnMatrixChanged().Broadcast(GetPixelMapping(), this);
+	GetOnMatrixChanged().Broadcast(PixelMapping, this);
 }
 
 void UDMXPixelMappingMatrixComponent::OnFixtureTypeChanged(const UDMXEntityFixtureType* FixtureType)
