DMX - Refactor rendering pixel mapping and preview rendering, improve performance and UX

diff --git a/Engine/Plugins/VirtualProduction/DMX/DMXPixelMapping/Source/DMXPixelMappingEditor/DMXPixelMappingEditor.Build.cs b/Engine/Plugins/VirtualProduction/DMX/DMXPixelMapping/Source/DMXPixelMappingEditor/DMXPixelMappingEditor.Build.cs
--- a/Engine/Plugins/VirtualProduction/DMX/DMXPixelMapping/Source/DMXPixelMappingEditor/DMXPixelMappingEditor.Build.cs
+++ b/Engine/Plugins/VirtualProduction/DMX/DMXPixelMapping/Source/DMXPixelMappingEditor/DMXPixelMappingEditor.Build.cs
@@ -37,6 +37,7 @@ public class DMXPixelMappingEditor : ModuleRules
 			"SlateCore",
 			"ToolMenus",
 			"ToolWidgets",
+			"UMG",
 			"UnrealEd",
 		});
 	}
diff --git a/Engine/Plugins/VirtualProduction/DMX/DMXPixelMapping/Source/DMXPixelMappingEditor/Private/Customizations/DMXPixelMappingDetailCustomization_MatrixCell.cpp b/Engine/Plugins/VirtualProduction/DMX/DMXPixelMapping/Source/DMXPixelMappingEditor/Private/Customizations/DMXPixelMappingDetailCustomization_MatrixCell.cpp
deleted file mode 100644
diff --git a/Engine/Plugins/VirtualProduction/DMX/DMXPixelMapping/Source/DMXPixelMappingEditor/Private/Customizations/DMXPixelMappingDetailCustomization_MatrixCell.h b/Engine/Plugins/VirtualProduction/DMX/DMXPixelMapping/Source/DMXPixelMappingEditor/Private/Customizations/DMXPixelMappingDetailCustomization_MatrixCell.h
deleted file mode 100644
diff --git a/Engine/Plugins/VirtualProduction/DMX/DMXPixelMapping/Source/DMXPixelMappingEditor/Private/Customizations/DMXPixelMappingDetailCustomization_Renderer.cpp b/Engine/Plugins/VirtualProduction/DMX/DMXPixelMapping/Source/DMXPixelMappingEditor/Private/Customizations/DMXPixelMappingDetailCustomization_Renderer.cpp
--- a/Engine/Plugins/VirtualProduction/DMX/DMXPixelMapping/Source/DMXPixelMappingEditor/Private/Customizations/DMXPixelMappingDetailCustomization_Renderer.cpp
+++ b/Engine/Plugins/VirtualProduction/DMX/DMXPixelMapping/Source/DMXPixelMappingEditor/Private/Customizations/DMXPixelMappingDetailCustomization_Renderer.cpp
@@ -1,153 +1,121 @@
 // Copyright Epic Games, Inc. All Rights Reserved.
 
 #include "Customizations/DMXPixelMappingDetailCustomization_Renderer.h"
-#include "Components/DMXPixelMappingRendererComponent.h"
-#include "DMXPixelMappingTypes.h"
 
+#include "Blueprint/UserWidget.h"
+#include "Components/DMXPixelMappingRendererComponent.h"
 #include "DetailLayoutBuilder.h"
 #include "DetailWidgetRow.h"
-#include "SWarningOrErrorBox.h"
+#include "DMXPixelMappingTypes.h"
 #include "Engine/Texture.h"
+#include "IPropertyUtilities.h"
 #include "Materials/Material.h"
+#include "SWarningOrErrorBox.h"
 #include "TextureResource.h"
 #include "Widgets/Text/STextBlock.h"
 
 #define LOCTEXT_NAMESPACE "DMXPixelMappingDetailCustomization_Renderer"
 
 
+FDMXPixelMappingDetailCustomization_Renderer::FDMXPixelMappingDetailCustomization_Renderer(const TWeakPtr<FDMXPixelMappingToolkit> InWeakToolkit)
+	: WeakToolkit(InWeakToolkit)
+{}
+
 void FDMXPixelMappingDetailCustomization_Renderer::CustomizeDetails(IDetailLayoutBuilder& DetailLayout)
 {
-	// Register all handles
-	RendererTypePropertyHandle = DetailLayout.GetProperty(GET_MEMBER_NAME_CHECKED(UDMXPixelMappingRendererComponent, RendererType));
-	InputTexturePropertyHandle = DetailLayout.GetProperty(GET_MEMBER_NAME_CHECKED(UDMXPixelMappingRendererComponent, InputTexture));
-	InputMaterialPropertyHandle = DetailLayout.GetProperty(GET_MEMBER_NAME_CHECKED(UDMXPixelMappingRendererComponent, InputMaterial));
-	InputWidgetPropertyHandle = DetailLayout.GetProperty(GET_MEMBER_NAME_CHECKED(UDMXPixelMappingRendererComponent, InputWidget));
+	PropertyUtilities = DetailLayout.GetPropertyUtilities();
 
-	// Bind property changes
-	FSimpleDelegate OnRendererTypePropertyPreChangeDelegate = FSimpleDelegate::CreateSP(this, &FDMXPixelMappingDetailCustomization_Renderer::OnRendererTypePropertyPreChange);
-	RendererTypePropertyHandle->SetOnPropertyValuePreChange(OnRendererTypePropertyPreChangeDelegate);
+	InputTextureHandle = DetailLayout.GetProperty(GET_MEMBER_NAME_CHECKED(UDMXPixelMappingRendererComponent, InputTexture));
+	InputMaterialHandle = DetailLayout.GetProperty(GET_MEMBER_NAME_CHECKED(UDMXPixelMappingRendererComponent, InputMaterial));
+	InputWidgetHandle = DetailLayout.GetProperty(GET_MEMBER_NAME_CHECKED(UDMXPixelMappingRendererComponent, InputWidget));
+	RendererTypeHandle = DetailLayout.GetProperty(GET_MEMBER_NAME_CHECKED(UDMXPixelMappingRendererComponent, RendererType));
 
-	FSimpleDelegate OnRendererTypePropertyChangedDelegate = FSimpleDelegate::CreateSP(this, &FDMXPixelMappingDetailCustomization_Renderer::OnRendererTypePropertyChanged);
-	RendererTypePropertyHandle->SetOnPropertyValueChanged(OnRendererTypePropertyChangedDelegate);
-
-	FSimpleDelegate OnInputTexturePropertyChangedDelegate = FSimpleDelegate::CreateSP(this, &FDMXPixelMappingDetailCustomization_Renderer::OnInputTexturePropertyChanged);
-	RendererTypePropertyHandle->SetOnPropertyValueChanged(OnInputTexturePropertyChangedDelegate);
 
 	// Hide Output Component properties
-	const TSharedRef<IPropertyHandle> PositionXPropertyHandle = DetailLayout.GetProperty(UDMXPixelMappingOutputComponent::GetPositionXPropertyName(), UDMXPixelMappingOutputComponent::StaticClass());
-	DetailLayout.HideProperty(PositionXPropertyHandle);
+	const TSharedRef<IPropertyHandle> PositionXHandle = DetailLayout.GetProperty(UDMXPixelMappingOutputComponent::GetPositionXPropertyName(), UDMXPixelMappingOutputComponent::StaticClass());
+	DetailLayout.HideProperty(PositionXHandle);
 
-	const TSharedRef<IPropertyHandle> PositionYPropertyHandle = DetailLayout.GetProperty(UDMXPixelMappingOutputComponent::GetPositionYPropertyName(), UDMXPixelMappingOutputComponent::StaticClass());
-	DetailLayout.HideProperty(PositionYPropertyHandle);
+	const TSharedRef<IPropertyHandle> PositionYHandle = DetailLayout.GetProperty(UDMXPixelMappingOutputComponent::GetPositionYPropertyName(), UDMXPixelMappingOutputComponent::StaticClass());
+	DetailLayout.HideProperty(PositionYHandle);
 
-	const TSharedRef<IPropertyHandle> SizeXPropertyHandle = DetailLayout.GetProperty(UDMXPixelMappingOutputComponent::GetSizeXPropertyName(), UDMXPixelMappingOutputComponent::StaticClass());
-	DetailLayout.HideProperty(SizeXPropertyHandle);
+	const TSharedRef<IPropertyHandle> SizeXHandle = DetailLayout.GetProperty(UDMXPixelMappingOutputComponent::GetSizeXPropertyName(), UDMXPixelMappingOutputComponent::StaticClass());
+	DetailLayout.HideProperty(SizeXHandle);
 
-	const TSharedRef<IPropertyHandle> SizeYPropertyHandle = DetailLayout.GetProperty(UDMXPixelMappingOutputComponent::GetSizeYPropertyName(), UDMXPixelMappingOutputComponent::StaticClass());
-	DetailLayout.HideProperty(SizeYPropertyHandle);
+	const TSharedRef<IPropertyHandle> SizeYHandle = DetailLayout.GetProperty(UDMXPixelMappingOutputComponent::GetSizeYPropertyName(), UDMXPixelMappingOutputComponent::StaticClass());
+	DetailLayout.HideProperty(SizeYHandle);
 
-	const TSharedRef<IPropertyHandle> IsLockInDesignerPropertyHandle = DetailLayout.GetProperty(UDMXPixelMappingOutputComponent::GetLockInDesignerPropertyName(), UDMXPixelMappingOutputComponent::StaticClass());
-	DetailLayout.HideProperty(IsLockInDesignerPropertyHandle);
+	const TSharedRef<IPropertyHandle> IsLockInDesignerHandle = DetailLayout.GetProperty(UDMXPixelMappingOutputComponent::GetLockInDesignerPropertyName(), UDMXPixelMappingOutputComponent::StaticClass());
+	DetailLayout.HideProperty(IsLockInDesignerHandle);
 
-	const TSharedRef<IPropertyHandle> IsVisibleInDesignerPropertyHandle = DetailLayout.GetProperty(UDMXPixelMappingOutputComponent::GetVisibleInDesignerPropertyName(), UDMXPixelMappingOutputComponent::StaticClass());
-	DetailLayout.HideProperty(IsVisibleInDesignerPropertyHandle);
+	const TSharedRef<IPropertyHandle> IsVisibleInDesignerHandle = DetailLayout.GetProperty(UDMXPixelMappingOutputComponent::GetVisibleInDesignerPropertyName(), UDMXPixelMappingOutputComponent::StaticClass());
+	DetailLayout.HideProperty(IsVisibleInDesignerHandle);
 
-	const TSharedRef<IPropertyHandle> PixelBlendingPropertyHandle = DetailLayout.GetProperty(GET_MEMBER_NAME_CHECKED(UDMXPixelMappingOutputComponent, CellBlendingQuality), UDMXPixelMappingOutputComponent::StaticClass());
-	DetailLayout.HideProperty(PixelBlendingPropertyHandle);
+	const TSharedRef<IPropertyHandle> PixelBlendingHandle = DetailLayout.GetProperty(GET_MEMBER_NAME_CHECKED(UDMXPixelMappingOutputComponent, CellBlendingQuality), UDMXPixelMappingOutputComponent::StaticClass());
+	DetailLayout.HideProperty(PixelBlendingHandle);
 
 	// Add properties
 	IDetailCategoryBuilder& RenderSettingsCategory = DetailLayout.EditCategory("Render Settings", FText::GetEmpty(), ECategoryPriority::Important);
-	RenderSettingsCategory.AddProperty(RendererTypePropertyHandle);
+	RenderSettingsCategory.AddProperty(RendererTypeHandle);
 
 	// Get editing UObject
-	TArray<TWeakObjectPtr<UObject>> OuterObjects;
-	DetailLayout.GetObjectsBeingCustomized(OuterObjects);
-	if (OuterObjects.Num() == 1)
+	TArray<TWeakObjectPtr<UObject>> CustomizedObjects;
+	DetailLayout.GetObjectsBeingCustomized(CustomizedObjects);
+
+	// Only show warnings if not multi-selecting
+	CustomizedObjects.Remove(nullptr);
+	if (CustomizedObjects.Num() == 1)
 	{
-		RendererComponent = Cast<UDMXPixelMappingRendererComponent>(OuterObjects[0]);
+		WeakRendererComponent = Cast<UDMXPixelMappingRendererComponent>(CustomizedObjects[0]);
 
-		if (RendererComponent.IsValid())
-		{
-			// Remember the selected renderer type
-			PreviousRendererType = RendererComponent->RendererType;
-
-			// Add Warning InputTexture Row
-			AddInputTextureWarning(RenderSettingsCategory);
-
-			// Add non UI Material Warning
-			AddMaterialWarning(RenderSettingsCategory);
-
-			// Register properties
-			RenderSettingsCategory.AddProperty(InputTexturePropertyHandle)
-				.Visibility(TAttribute<EVisibility>::Create(TAttribute<EVisibility>::FGetter::CreateLambda([this]
-					{
-						return IsSelectedRendererType(EDMXPixelMappingRendererType::Texture) ? EVisibility::Visible : EVisibility::Hidden;
-					})));
-
-			RenderSettingsCategory.AddProperty(InputMaterialPropertyHandle)
-				.Visibility(TAttribute<EVisibility>::Create(TAttribute<EVisibility>::FGetter::CreateLambda([this]
-					{
-						return IsSelectedRendererType(EDMXPixelMappingRendererType::Material) ? EVisibility::Visible : EVisibility::Hidden;
-					})));
-
-			RenderSettingsCategory.AddProperty(InputWidgetPropertyHandle)
-				.Visibility(TAttribute<EVisibility>::Create(TAttribute<EVisibility>::FGetter::CreateLambda([this]
-					{
-						return IsSelectedRendererType(EDMXPixelMappingRendererType::UMG) ? EVisibility::Visible : EVisibility::Hidden;
-					})));
-		}
-	}
-}
+		// Add Warning InputTexture Row
+		AddInputTextureWarning(RenderSettingsCategory);
 
-void FDMXPixelMappingDetailCustomization_Renderer::OnRendererTypePropertyPreChange()
-{
-	if (UDMXPixelMappingRendererComponent* Component = RendererComponent.Get())
-	{
-		PreviousRendererType = Component->RendererType;
-	}
-}
+		// Add non UI Material Warning
+		AddMaterialDomainWarning(RenderSettingsCategory);
 
-void FDMXPixelMappingDetailCustomization_Renderer::OnRendererTypePropertyChanged()
-{
-	if (UDMXPixelMappingRendererComponent* Component = RendererComponent.Get())
-	{
-		if (Component->RendererType != PreviousRendererType)
-		{
-			PreviousRendererType = Component->RendererType;
-			ResetSize();
-		}
-	}
-}
+		// Register properties
+		RenderSettingsCategory.AddProperty(InputTextureHandle)
+			.Visibility(TAttribute<EVisibility>::Create(TAttribute<EVisibility>::FGetter::CreateLambda([this]
+				{
+					return GetRendererType() == EDMXPixelMappingRendererType::Texture ? EVisibility::Visible : EVisibility::Hidden;
+				})));
 
-void FDMXPixelMappingDetailCustomization_Renderer::OnInputTexturePropertyChanged()
-{
-	ResetSize();
+		RenderSettingsCategory.AddProperty(InputMaterialHandle)
+			.Visibility(TAttribute<EVisibility>::Create(TAttribute<EVisibility>::FGetter::CreateLambda([this]
+				{
+					return GetRendererType() == EDMXPixelMappingRendererType::Material ? EVisibility::Visible : EVisibility::Hidden;
+				})));
+
+		RenderSettingsCategory.AddProperty(InputWidgetHandle)
+			.Visibility(TAttribute<EVisibility>::Create(TAttribute<EVisibility>::FGetter::CreateLambda([this]
+				{
+					return GetRendererType() == EDMXPixelMappingRendererType::UMG ? EVisibility::Visible : EVisibility::Hidden;
+				})));
+	}
 }
 
-bool FDMXPixelMappingDetailCustomization_Renderer::IsSelectedRendererType(EDMXPixelMappingRendererType PropertyRendererType) const
+void FDMXPixelMappingDetailCustomization_Renderer::PostUndo(bool bSuccess)
 {
-	if (UDMXPixelMappingRendererComponent* Component = RendererComponent.Get())
+	const TArray<TWeakObjectPtr<UObject>> SelectedObjects = PropertyUtilities->GetSelectedObjects();
+	for (const TWeakObjectPtr<UObject>& Object : SelectedObjects)
 	{
-		if (Component->RendererType == PropertyRendererType)
+		if (UDMXPixelMappingRendererComponent* RendererComponent = Cast<UDMXPixelMappingRendererComponent>(Object))
 		{
-			return true;
+			RendererComponent->InvalidatePreprocessRenderer();
 		}
 	}
-
-	return false;
 }
 
-bool FDMXPixelMappingDetailCustomization_Renderer::IsNotSelectedRendererType(EDMXPixelMappingRendererType PropertyRendererType) const
+void FDMXPixelMappingDetailCustomization_Renderer::PostRedo(bool bSuccess)
 {
-	if (UDMXPixelMappingRendererComponent* Component = RendererComponent.Get())
+	const TArray<TWeakObjectPtr<UObject>> SelectedObjects = PropertyUtilities->GetSelectedObjects();
+	for (const TWeakObjectPtr<UObject>& Object : SelectedObjects)
 	{
-		if (Component->RendererType != PropertyRendererType)
+		if (UDMXPixelMappingRendererComponent* RendererComponent = Cast<UDMXPixelMappingRendererComponent>(Object))
 		{
-			return true;
+			RendererComponent->InvalidatePreprocessRenderer();
 		}
 	}
-
-	return false;
 }
 
 void FDMXPixelMappingDetailCustomization_Renderer::AddInputTextureWarning(IDetailCategoryBuilder& InCategory)
@@ -155,7 +123,7 @@ void FDMXPixelMappingDetailCustomization_Renderer::AddInputTextureWarning(IDetai
 	const FSlateBrush* WarningIcon = FAppStyle::GetBrush("SettingsEditor.WarningIcon");
 
 	InCategory.AddCustomRow(FText::GetEmpty())
-		.Visibility(TAttribute<EVisibility>(this, &FDMXPixelMappingDetailCustomization_Renderer::GetInputTextureWarning))
+		.Visibility(TAttribute<EVisibility>(this, &FDMXPixelMappingDetailCustomization_Renderer::GetInputTextureWarningVisibility))
 		.WholeRowContent()
 		[
 			SNew(SWarningOrErrorBox)
@@ -164,10 +132,10 @@ void FDMXPixelMappingDetailCustomization_Renderer::AddInputTextureWarning(IDetai
 		];
 }
 
-void FDMXPixelMappingDetailCustomization_Renderer::AddMaterialWarning(IDetailCategoryBuilder& InCategory)
+void FDMXPixelMappingDetailCustomization_Renderer::AddMaterialDomainWarning(IDetailCategoryBuilder& InCategory)
 {
 	InCategory.AddCustomRow(FText::GetEmpty())
-		.Visibility(TAttribute<EVisibility>(this, &FDMXPixelMappingDetailCustomization_Renderer::GetMaterialWarningVisibility))
+		.Visibility(TAttribute<EVisibility>(this, &FDMXPixelMappingDetailCustomization_Renderer::GetMaterialDomainWarningVisibility))
 		.WholeRowContent()
 		[
 			SNew(SWarningOrErrorBox)
@@ -176,97 +144,82 @@ void FDMXPixelMappingDetailCustomization_Renderer::AddMaterialWarning(IDetailCat
 		];
 }
 
-EVisibility FDMXPixelMappingDetailCustomization_Renderer::GetMaterialWarningVisibility() const
+
+EVisibility FDMXPixelMappingDetailCustomization_Renderer::GetInputTextureWarningVisibility() const
 {
-	if (UDMXPixelMappingRendererComponent* Component = RendererComponent.Get())
+	UDMXPixelMappingRendererComponent* RendererComponent = WeakRendererComponent.Get();
+	if (!RendererComponent)
 	{
-		if (Component->InputMaterial != nullptr)
-		{
-			if (UMaterial* Material = Component->InputMaterial->GetMaterial())
-			{
-				if (Component->RendererType == EDMXPixelMappingRendererType::Material && !Material->IsUIMaterial())
-				{
-					return EVisibility::Visible;
-				}
-			}
-		}
+		return EVisibility::Collapsed;
 	}
 
-	return EVisibility::Collapsed;
+	const EDMXPixelMappingRendererType RendererType = GetRendererType();
+	if (RendererType == EDMXPixelMappingRendererType::Texture)
+	{
+		return RendererComponent->InputTexture ? EVisibility::Collapsed : EVisibility::Visible;
+	}
+	else if (RendererType == EDMXPixelMappingRendererType::Material)
+	{
+		return RendererComponent->InputMaterial ? EVisibility::Collapsed : EVisibility::Visible;
+	}
+	else if (RendererType == EDMXPixelMappingRendererType::UMG)
+	{
+		return RendererComponent->InputWidget ? EVisibility::Collapsed : EVisibility::Visible;
+	}
+	else
+	{
+		ensureMsgf(0, TEXT("Renderer type not implemented in FDMXPixelMappingDetailCustomization_Renderer"));
+		return EVisibility::Collapsed;
+	}
 }
 
-EVisibility FDMXPixelMappingDetailCustomization_Renderer::GetInputTextureWarning() const
+FText FDMXPixelMappingDetailCustomization_Renderer::GetInputTextureWarningText() const
 {
-	if (UDMXPixelMappingRendererComponent* Component = RendererComponent.Get())
+	const EDMXPixelMappingRendererType RendererType = GetRendererType();
+	if (RendererType == EDMXPixelMappingRendererType::Texture)
 	{
-		if (!Component->GetRenderedInputTexture())
-		{
-			return EVisibility::Visible;
-		}
+		return LOCTEXT("WarningCategoryDisplayName.TextureNotSet", "Texture is not set.");
+	}
+	else if (RendererType == EDMXPixelMappingRendererType::Material)
+	{
+		return LOCTEXT("WarningCategoryDisplayName.MaterialNotSet", "Material is not set.");
+	}
+	else if (RendererType == EDMXPixelMappingRendererType::UMG)
+	{
+		return LOCTEXT("WarningCategoryDisplayName.UMGNotSet", "UMG is not set.");
+	}
+	else
+	{
+		ensureMsgf(0, TEXT("Renderer type not implemented in FDMXPixelMappingDetailCustomization_Renderer"));
+		return FText::GetEmpty();
 	}
-
-	return EVisibility::Collapsed;
 }
 
-FText FDMXPixelMappingDetailCustomization_Renderer::GetInputTextureWarningText() const
+EVisibility FDMXPixelMappingDetailCustomization_Renderer::GetMaterialDomainWarningVisibility() const
 {
-	if (UDMXPixelMappingRendererComponent* Component = RendererComponent.Get())
+	UDMXPixelMappingRendererComponent* RendererComponent = WeakRendererComponent.Get();
+	if (!RendererComponent || RendererComponent->RendererType != EDMXPixelMappingRendererType::Material)
 	{
-		if (!Component->GetRenderedInputTexture())
-		{
-			if (Component->RendererType == EDMXPixelMappingRendererType::Texture)
-			{
-				return LOCTEXT("WarningCategoryDisplayName.TextureNotSet", "Texture is not set.");
-			}
-			else if (Component->RendererType == EDMXPixelMappingRendererType::Material)
-			{
-				return LOCTEXT("WarningCategoryDisplayName.MaterialNotSet", "Material is not set.");
-			}
-			else if (Component->RendererType == EDMXPixelMappingRendererType::UMG)
-			{
-				return LOCTEXT("WarningCategoryDisplayName.UMGNotSet", "UMG is not set.");
-			}
-		}
+		return EVisibility::Collapsed;
 	}
 
-	return FText();
+	UMaterial* Material = RendererComponent->InputMaterial ? RendererComponent->InputMaterial->GetMaterial() : nullptr;
+	if (Material && !Material->IsUIMaterial())
+	{
+		return EVisibility::Visible;
+	}
+
+	return EVisibility::Collapsed;
 }
 
-void FDMXPixelMappingDetailCustomization_Renderer::ResetSize()
+EDMXPixelMappingRendererType FDMXPixelMappingDetailCustomization_Renderer::GetRendererType() const
 {
-	// Set a new size depending on the selected type
-	if (UDMXPixelMappingRendererComponent* Component = RendererComponent.Get())
-	{
-		if (Component->RendererType == EDMXPixelMappingRendererType::Texture)
-		{
-			if (const UTexture* InputTexture = Component->InputTexture)
-			{
-				if (const FTextureResource* TextureResource = InputTexture->GetResource())
-				{
-					// Set to the texture size
-					const FVector2D NewSize = FVector2D(TextureResource->GetSizeX(), TextureResource->GetSizeY());
-					Component->SetSize(NewSize);
-					return;
-				}
-			}
-
-			// Reset to default size
-			const FVector2D DefaultSize = FVector2D(100.f, 100.f);
-			Component->SetSize(DefaultSize);
-		}
-		else if (Component->RendererType == EDMXPixelMappingRendererType::Material)
-		{
-			// Reset to default size
-			const FVector2D DefaultSize = FVector2D(100.f, 100.f);
-			Component->SetSize(DefaultSize);
-		}
-		else if (Component->RendererType == EDMXPixelMappingRendererType::UMG)
-		{
-			// Reset to default size
-			const FVector2D DefaultSize = FVector2D(1024.f, 768.f);
-			Component->SetSize(DefaultSize);
-		}
-	}
+	TArray<void*> RawDataArray;
+	RendererTypeHandle->AccessRawData(RawDataArray);
+
+	return RawDataArray.IsEmpty() ?
+		EDMXPixelMappingRendererType::Texture :
+		*static_cast<EDMXPixelMappingRendererType*>(RawDataArray[0]);
 }
 
 #undef LOCTEXT_NAMESPACE
diff --git a/Engine/Plugins/VirtualProduction/DMX/DMXPixelMapping/Source/DMXPixelMappingEditor/Private/Customizations/DMXPixelMappingDetailCustomization_Renderer.h b/Engine/Plugins/VirtualProduction/DMX/DMXPixelMapping/Source/DMXPixelMappingEditor/Private/Customizations/DMXPixelMappingDetailCustomization_Renderer.h
--- a/Engine/Plugins/VirtualProduction/DMX/DMXPixelMapping/Source/DMXPixelMappingEditor/Private/Customizations/DMXPixelMappingDetailCustomization_Renderer.h
+++ b/Engine/Plugins/VirtualProduction/DMX/DMXPixelMapping/Source/DMXPixelMappingEditor/Private/Customizations/DMXPixelMappingDetailCustomization_Renderer.h
@@ -2,72 +2,71 @@
 
 #pragma once
 
+#include "EditorUndoClient.h"
 #include "IDetailCustomization.h"
 #include "UObject/WeakObjectPtr.h"
 
+enum class EDMXPixelMappingRendererType : uint8;
+struct EVisibility;
 class FDMXPixelMappingToolkit;
-class IPropertyHandle;
+class FText;
 class IDetailCategoryBuilder;
+class IPropertyHandle;
+class IPropertyUtilities;
 class UDMXPixelMappingRendererComponent;
-struct EVisibility;
-enum class EDMXPixelMappingRendererType : uint8;
+
 
 class FDMXPixelMappingDetailCustomization_Renderer
 	: public IDetailCustomization
+	, public FSelfRegisteringEditorUndoClient
 {
 public:
-	static TSharedRef<IDetailCustomization> MakeInstance(const TWeakPtr<FDMXPixelMappingToolkit> InToolkitWeakPtr)
+	FDMXPixelMappingDetailCustomization_Renderer(const TWeakPtr<FDMXPixelMappingToolkit> InWeakToolkit);
+
+	static TSharedRef<IDetailCustomization> MakeInstance(const TWeakPtr<FDMXPixelMappingToolkit> InWeakToolkit)
 	{
-		return MakeShared<FDMXPixelMappingDetailCustomization_Renderer>(InToolkitWeakPtr);
+		return MakeShared<FDMXPixelMappingDetailCustomization_Renderer>(InWeakToolkit);
 	}
 
-	/** Constructor */
-	FDMXPixelMappingDetailCustomization_Renderer(const TWeakPtr<FDMXPixelMappingToolkit> InToolkitWeakPtr)
-		: ToolkitWeakPtr(InToolkitWeakPtr)
-	{}
-
 	/** IDetailCustomization interface */
 	virtual void CustomizeDetails(IDetailLayoutBuilder& DetailLayout) override;
 
-private:
-	/** Called before the renderer type property changed */
-	void OnRendererTypePropertyPreChange();
-
-	/** Called when the renderer type property changed */
-	void OnRendererTypePropertyChanged();
-
-	/** Called when the input texture property changed */
-	void OnInputTexturePropertyChanged();
+protected:
+	//~ Begin FSelfRegisteringEditorUndoClient interface
+	virtual void PostUndo(bool bSuccess) override;
+	virtual void PostRedo(bool bSuccess) override;
+	//~ End FSelfRegisteringEditorUndoClient interface
 
-	bool IsSelectedRendererType(EDMXPixelMappingRendererType PropertyRendererType) const;
-
-	bool IsNotSelectedRendererType(EDMXPixelMappingRendererType PropertyRendererType) const;
+private:
+	/** Adds a warning if the input texture is not set */
+	void AddInputTextureWarning(IDetailCategoryBuilder& InCategory);
 
-	EVisibility GetInputTextureWarning() const;
+	/** Add warning for material with domain != MD_UI */
+	void AddMaterialDomainWarning(IDetailCategoryBuilder& InCategory);
 
+	/** Returns the text displayed if the input texture is not set */
 	FText GetInputTextureWarningText() const;
 
-	void AddInputTextureWarning(IDetailCategoryBuilder& InCategory);
+	/** Returns the visiblity of the input texture warning text */
+	EVisibility GetInputTextureWarningVisibility() const;
 
-	/** Add warning for material with domain != MD_UI */
-	void AddMaterialWarning(IDetailCategoryBuilder& InCategory);
+	/** Returns the visibility of the input material domain */
+	EVisibility GetMaterialDomainWarningVisibility() const;
 
-	/** Visibility for non ui material warning */
-	EVisibility GetMaterialWarningVisibility() const;
+	/** Returns the current renderer type */
+	EDMXPixelMappingRendererType GetRendererType() const;
 
-	/** Resets the size of the renderer component depending on current renderer type and the resource object */
-	void ResetSize();
+	TSharedPtr<IPropertyHandle> RendererTypeHandle;
+	TSharedPtr<IPropertyHandle> InputTextureHandle;
+	TSharedPtr<IPropertyHandle> InputMaterialHandle;
+	TSharedPtr<IPropertyHandle> InputWidgetHandle;
 
-	/** Previously selected renderer type */
-	EDMXPixelMappingRendererType PreviousRendererType;
-	
-	TSharedPtr<IPropertyHandle> RendererTypePropertyHandle;
-	TSharedPtr<IPropertyHandle> InputTexturePropertyHandle;
-	TSharedPtr<IPropertyHandle> InputMaterialPropertyHandle;
-	TSharedPtr<IPropertyHandle> InputWidgetPropertyHandle;
+	/** Renderer components being customized. Only valid if not multi editing */
+	TWeakObjectPtr<UDMXPixelMappingRendererComponent> WeakRendererComponent;
 
-	TWeakObjectPtr<UDMXPixelMappingRendererComponent> RendererComponent;
+	/** Property utilities for this customization */
+	TSharedPtr<IPropertyUtilities> PropertyUtilities;
 
 	/** Weak reference to the DMX editor */
-	TWeakPtr<FDMXPixelMappingToolkit> ToolkitWeakPtr;
+	TWeakPtr<FDMXPixelMappingToolkit> WeakToolkit;
 };
diff --git a/Engine/Plugins/VirtualProduction/DMX/DMXPixelMapping/Source/DMXPixelMappingEditor/Private/DMXPixelMappingEditorCommands.cpp b/Engine/Plugins/VirtualProduction/DMX/DMXPixelMapping/Source/DMXPixelMappingEditor/Private/DMXPixelMappingEditorCommands.cpp
--- a/Engine/Plugins/VirtualProduction/DMX/DMXPixelMapping/Source/DMXPixelMappingEditor/Private/DMXPixelMappingEditorCommands.cpp
+++ b/Engine/Plugins/VirtualProduction/DMX/DMXPixelMapping/Source/DMXPixelMappingEditor/Private/DMXPixelMappingEditorCommands.cpp
@@ -32,9 +32,10 @@ void FDMXPixelMappingEditorCommands::RegisterCommands()
 	UI_COMMAND(ToggleScaleChildrenWithParent, "Scale Children with Parent", "Sets if children are scaled with their parent, when the parent is resized.", EUserInterfaceActionType::ToggleButton, FInputChord(EModifierKey::Alt, EKeys::Q));
 	UI_COMMAND(ToggleAlwaysSelectGroup, "Always select Group", "Sets if the parent is selected, if a child is clicked.", EUserInterfaceActionType::ToggleButton, FInputChord(EModifierKey::Alt, EKeys::W));
 	UI_COMMAND(ToggleApplyLayoutScriptWhenLoaded, "Apply Layout Script instantly", "Sets if layout script are applied as soon as they are loaded.", EUserInterfaceActionType::ToggleButton, FInputChord(EModifierKey::Alt, EKeys::E));
-	UI_COMMAND(ToggleShowComponentNames, "Show Component Names", "Sets if the name of components are displayed. Note, may impact editor performance.", EUserInterfaceActionType::ToggleButton, FInputChord(EModifierKey::Alt, EKeys::A));
-	UI_COMMAND(ToggleShowPatchInfo, "Show Patch Info", "Sets if information about the pach is displayed. Note, may impact editor performance.", EUserInterfaceActionType::ToggleButton, FInputChord(EModifierKey::Alt, EKeys::S));
-	UI_COMMAND(ToggleShowCellIDs, "Show Cell IDs", "Sets if the cell IDs of matrix cells are displayed.", EUserInterfaceActionType::ToggleButton, FInputChord(EModifierKey::Alt, EKeys::D));
+	UI_COMMAND(ToggleShowMatrixCells, "Show Matrix Cells", "Sets if matrix cells are displayed. It is recommended that this is turned off when pixel mapping large quantities of fixtures.", EUserInterfaceActionType::ToggleButton, FInputChord(EModifierKey::Alt, EKeys::A));
+	UI_COMMAND(ToggleShowComponentNames, "Show Component Names", "Sets if the name of components are displayed. It is recommended that this is turned off when pixel mapping large quantities of fixtures.", EUserInterfaceActionType::ToggleButton, FInputChord(EModifierKey::Alt, EKeys::S));
+	UI_COMMAND(ToggleShowPatchInfo, "Show Patch Info", "Sets if information about the pach is displayed. It is recommended that this is turned off when pixel mapping large quantities of fixtures.", EUserInterfaceActionType::ToggleButton, FInputChord(EModifierKey::Alt, EKeys::D));
+	UI_COMMAND(ToggleShowCellIDs, "Show Cell IDs", "Sets if the cell IDs of matrix cells are displayed. It is recommended that this is turned off when pixel mapping large quantities of fixtures.", EUserInterfaceActionType::ToggleButton, FInputChord(EModifierKey::Alt, EKeys::F));
 }
 
 #undef LOCTEXT_NAMESPACE 
diff --git a/Engine/Plugins/VirtualProduction/DMX/DMXPixelMapping/Source/DMXPixelMappingEditor/Private/DMXPixelMappingEditorCommands.h b/Engine/Plugins/VirtualProduction/DMX/DMXPixelMapping/Source/DMXPixelMappingEditor/Private/DMXPixelMappingEditorCommands.h
--- a/Engine/Plugins/VirtualProduction/DMX/DMXPixelMapping/Source/DMXPixelMappingEditor/Private/DMXPixelMappingEditorCommands.h
+++ b/Engine/Plugins/VirtualProduction/DMX/DMXPixelMapping/Source/DMXPixelMappingEditor/Private/DMXPixelMappingEditorCommands.h
@@ -26,6 +26,7 @@ public:
 	TSharedPtr<FUICommandInfo> ToggleScaleChildrenWithParent;
 	TSharedPtr<FUICommandInfo> ToggleAlwaysSelectGroup;
 	TSharedPtr<FUICommandInfo> ToggleApplyLayoutScriptWhenLoaded;
+	TSharedPtr<FUICommandInfo> ToggleShowMatrixCells;
 	TSharedPtr<FUICommandInfo> ToggleShowComponentNames;
 	TSharedPtr<FUICommandInfo> ToggleShowPatchInfo;
 	TSharedPtr<FUICommandInfo> ToggleShowCellIDs;
diff --git a/Engine/Plugins/VirtualProduction/DMX/DMXPixelMapping/Source/DMXPixelMappingEditor/Private/DMXPixelMappingEditorStyle.cpp b/Engine/Plugins/VirtualProduction/DMX/DMXPixelMapping/Source/DMXPixelMappingEditor/Private/DMXPixelMappingEditorStyle.cpp
--- a/Engine/Plugins/VirtualProduction/DMX/DMXPixelMapping/Source/DMXPixelMappingEditor/Private/DMXPixelMappingEditorStyle.cpp
+++ b/Engine/Plugins/VirtualProduction/DMX/DMXPixelMapping/Source/DMXPixelMappingEditor/Private/DMXPixelMappingEditorStyle.cpp
@@ -83,6 +83,7 @@ TSharedRef<FSlateStyleSet> FDMXPixelMappingEditorStyle::Create()
 	Style->Set("DMXPixelMappingEditor.StopPlayingDMX", new IMAGE_BRUSH("icon_DMXPixelMappingEditor_StopPlayingDMX_40x", Icon40x40));
 	Style->Set("DMXPixelMappingEditor.StopPlayingDMX.Small", new IMAGE_BRUSH("icon_DMXPixelMappingEditor_StopPlayingDMX_40x", Icon20x20));
 	
+	Style->Set("DMXPixelMappingEditor.ComponentBorder", new BORDER_BRUSH("Common/Selection", FMargin(2.f), FLinearColor::White));
 
 	return Style;
 }
diff --git a/Engine/Plugins/VirtualProduction/DMX/DMXPixelMapping/Source/DMXPixelMappingEditor/Private/DMXPixelMappingLayoutSettings.h b/Engine/Plugins/VirtualProduction/DMX/DMXPixelMapping/Source/DMXPixelMappingEditor/Private/DMXPixelMappingLayoutSettings.h
--- a/Engine/Plugins/VirtualProduction/DMX/DMXPixelMapping/Source/DMXPixelMappingEditor/Private/DMXPixelMappingLayoutSettings.h
+++ b/Engine/Plugins/VirtualProduction/DMX/DMXPixelMapping/Source/DMXPixelMappingEditor/Private/DMXPixelMappingLayoutSettings.h
@@ -27,15 +27,19 @@ public:
 	UPROPERTY(Config, EditAnywhere, Category = "Layout Settings")
 	bool bApplyLayoutScriptWhenLoaded = true;
 
-	/** If true, shows the name of the Fixture Patch where applicable. May affect editor performance. */
+	/**  If true, shows a widget for each cell. It is recommended that this is turned off when pixel mapping large quantities of fixtures. */
+	UPROPERTY(Config, EditAnywhere, Category = "Layout Settings")
+	bool bShowMatrixCells = false;
+
+	/** If true, shows the name of the Fixture Patch where applicable. It is recommended that this is turned off when pixel mapping large quantities of fixtures. */
 	UPROPERTY(Config, EditAnywhere, Category = "Layout Settings")
 	bool bShowComponentNames = true;
 
-	/** If true, shows Fixture info bout the Fixture Patch where applicable. May affect editor performance. */
+	/** If true, shows Fixture info bout the Fixture Patch where applicable. It is recommended that this is turned off when pixel mapping large quantities of fixtures. */
 	UPROPERTY(Config, EditAnywhere, Category = "Layout Settings")
 	bool bShowPatchInfo = false;
 
-	/** If true, shows Cell IDs where applicable. May affect editor performance. */
+	/** If true, shows Cell IDs where applicable. It is recommended that this is turned off when pixel mapping large quantities of fixtures. */
 	UPROPERTY(Config, EditAnywhere, Category = "Layout Settings")
 	bool bShowCellIDs = false;
 
diff --git a/Engine/Plugins/VirtualProduction/DMX/DMXPixelMapping/Source/DMXPixelMappingEditor/Private/DragDrop/DMXPixelMappingDragDropOp.cpp b/Engine/Plugins/VirtualProduction/DMX/DMXPixelMapping/Source/DMXPixelMappingEditor/Private/DragDrop/DMXPixelMappingDragDropOp.cpp
--- a/Engine/Plugins/VirtualProduction/DMX/DMXPixelMapping/Source/DMXPixelMappingEditor/Private/DragDrop/DMXPixelMappingDragDropOp.cpp
+++ b/Engine/Plugins/VirtualProduction/DMX/DMXPixelMapping/Source/DMXPixelMappingEditor/Private/DragDrop/DMXPixelMappingDragDropOp.cpp
@@ -85,13 +85,13 @@ void FDMXPixelMappingDragDropOp::LayoutOutputComponents(const FVector2D& GraphSp
 			{
 				if (UDMXPixelMappingOutputComponent* OutputComponent = Cast<UDMXPixelMappingOutputComponent>(Component.Get()))
 				{
-					OutputComponent->Modify();
+					OutputComponent->PreEditChange(nullptr);
 
-					constexpr bool bModifyChildrenRecursively = true;
+					constexpr bool bNotifyChildrenRecursive = true;
 					Component->ForEachChild([](UDMXPixelMappingBaseComponent* Component)
-				{
-							Component->Modify();
-						}, bModifyChildrenRecursively);
+						{
+							Component->PreEditChange(nullptr);
+						}, bNotifyChildrenRecursive);
 
 					if (ensureMsgf(OutputComponent->GetClass() != UDMXPixelMappingMatrixComponent::StaticClass(),
 						TEXT("Matrix Cells are not supported. Use the GroupChildDragDropHelper from this class instead")))
@@ -100,7 +100,15 @@ void FDMXPixelMappingDragDropOp::LayoutOutputComponents(const FVector2D& GraphSp
 
 						const FVector2D NewPosition = GraphSpacePosition - AnchorOffset - GraphSpaceDragOffset;
 						OutputComponent->SetPosition(NewPosition.RoundToVector());
+
+						OutputComponent->PostEditChange();
 					}
+
+					Component->ForEachChild([](UDMXPixelMappingBaseComponent* Component)
+						{
+							Component->PostEditChange();
+						}, bNotifyChildrenRecursive);
+
 				}
 			}
 		}
diff --git a/Engine/Plugins/VirtualProduction/DMX/DMXPixelMapping/Source/DMXPixelMappingEditor/Private/Toolkits/DMXPixelMappingToolkit.cpp b/Engine/Plugins/VirtualProduction/DMX/DMXPixelMapping/Source/DMXPixelMappingEditor/Private/Toolkits/DMXPixelMappingToolkit.cpp
--- a/Engine/Plugins/VirtualProduction/DMX/DMXPixelMapping/Source/DMXPixelMappingEditor/Private/Toolkits/DMXPixelMappingToolkit.cpp
+++ b/Engine/Plugins/VirtualProduction/DMX/DMXPixelMapping/Source/DMXPixelMappingEditor/Private/Toolkits/DMXPixelMappingToolkit.cpp
@@ -150,68 +150,16 @@ void FDMXPixelMappingToolkit::Tick(float DeltaTime)
 		return;
 	}
 
-	TSet<UDMXPixelMappingRendererComponent*> RenderedRendererComponents;
-
-	// Render preview if required
-	for (const FDMXPixelMappingComponentReference& SelectedComponentRef : SelectedComponents)
-	{
-		UDMXPixelMappingBaseComponent* SelectedComponent = SelectedComponentRef.Component.Get();
-		if (IsValid(SelectedComponent))
-		{
-			// User select root component
-			if (Cast<UDMXPixelMappingRootComponent>(SelectedComponent))
-			{
-				break;
-			}
-
-			// Try to get renderer component from selected component
-			UDMXPixelMappingRendererComponent* RendererComponent = SelectedComponent->GetRendererComponent();
-			if (!ensureMsgf(RendererComponent, TEXT("Component %s resides in pixelmapping but has no valid renderer."), *SelectedComponent->GetUserFriendlyName()))
-			{
-				break;
-			}
-
-			// Render preview depends on the render component being 'rendered'
-			if (!RenderedRendererComponents.Contains(RendererComponent))
-			{
-				RendererComponent->Render();
-				RenderedRendererComponents.Add(RendererComponent);
-			}
-
-			// Render preview
-			RendererComponent->RenderEditorPreviewTexture();
-
-
-			// Render only once for all selected components
-			break;
-		}
-	}
-
 	// Render Output and Send DMX if required
-	if (bIsPlayingDMX && !bRequestStopSendingDMX)
+	if (!bRequestStopSendingDMX)
 	{
-		if (bTogglePlayDMXAll) // Send to all
+		if (bIsPlayingDMX)
 		{
-			// Render all components
 			RootComponent->RenderAndSendDMX();
 		}
-		else // Send to selected component
+		else
 		{
-			for (FDMXPixelMappingComponentReference& SelectedComponentRef : SelectedComponents)
-			{
-				if (UDMXPixelMappingBaseComponent* SelectedComponent = SelectedComponentRef.Component.Get())
-				{
-					// Try to get renderer component from selected component
-					UDMXPixelMappingRendererComponent* RendererComponent = SelectedComponent->GetRendererComponent();
-					if (RendererComponent && !RenderedRendererComponents.Contains(RendererComponent))
-					{
-						RendererComponent->Render();
-						RenderedRendererComponents.Add(RendererComponent);
-					}
-
-					SelectedComponent->SendDMX();
-				}
-			}
+			RootComponent->Render();
 		}
 	}
 	else if (bRequestStopSendingDMX)
@@ -335,7 +283,7 @@ void FDMXPixelMappingToolkit::OnSaveThumbnailImage()
 {
 	if (ActiveRendererComponent.IsValid())
 	{
-		DMXPixelMapping->ThumbnailImage = ActiveRendererComponent->GetPreviewRenderTarget();
+		DMXPixelMapping->ThumbnailImage = ActiveRendererComponent->GetRenderedInputTexture();
 	}
 }
 
@@ -370,49 +318,50 @@ TArray<UDMXPixelMappingBaseComponent*> FDMXPixelMappingToolkit::CreateComponents
 
 void FDMXPixelMappingToolkit::DeleteSelectedComponents()
 {
-	if (SelectedComponents.Num() > 0)
+	if (SelectedComponents.IsEmpty())
 	{
-		TGuardValue<bool>(bRemovingComponents, true);
+		return;
+	}
 
-		TSet<FDMXPixelMappingComponentReference> ParentComponentReferences;
-		for (const FDMXPixelMappingComponentReference& ComponentReference : SelectedComponents)
-		{
-			if (UDMXPixelMappingBaseComponent* Component = ComponentReference.GetComponent())
-			{
-				constexpr bool bModifyChildrenRecursively = true;
-				Component->ForEachChild([](UDMXPixelMappingBaseComponent* ChildComponent)
-					{
-						ChildComponent->Modify();
-					}, bModifyChildrenRecursively);
+	TGuardValue<bool>(bRemovingComponents, true);
 
-				if (UDMXPixelMappingBaseComponent* ParentComponent = Component->GetParent())
+	TSet<FDMXPixelMappingComponentReference> ParentComponentReferences;
+	for (const FDMXPixelMappingComponentReference& SelectedComponentReference : SelectedComponents)
+	{
+		UDMXPixelMappingBaseComponent* SelectedComponent = SelectedComponentReference.GetComponent();
+		if (SelectedComponent)
+		{
+			constexpr bool bModifyChildrenRecursively = true;
+			SelectedComponent->ForEachChild([](UDMXPixelMappingBaseComponent* ChildComponent)
 				{
-					ParentComponentReferences.Add(GetReferenceFromComponent(ParentComponent));
+					ChildComponent->Modify();
+				}, bModifyChildrenRecursively);
 
-					Component->SetFlags(RF_Transactional);
-					Component->Modify();
-					ParentComponent->Modify();
+			UDMXPixelMappingBaseComponent* ParentComponent = SelectedComponent->GetParent();
+			if (ParentComponent)
+			{
+				ParentComponent->Modify();
+				SelectedComponent->Modify();
 
-					ParentComponent->RemoveChild(Component);
-				}
-			
-				if (UDMXPixelMappingOutputComponent* OutputComponent = Cast<UDMXPixelMappingOutputComponent>(Component))
-				{
-					ActiveOutputComponents.Remove(OutputComponent);
-				}
+				ParentComponent->RemoveChild(SelectedComponent);
+
+				const bool bParentComponentIsBeingRemoved = Algo::FindByPredicate(SelectedComponents, [ParentComponent](const FDMXPixelMappingComponentReference& Reference)
+					{
+						return Reference.GetComponent() == ParentComponent;
+					}) != nullptr;
 
-				if (UDMXPixelMappingRendererComponent* RendererComponent = Cast<UDMXPixelMappingRendererComponent>(Component))
+				if (!bParentComponentIsBeingRemoved)
 				{
-					SetActiveRenderComponent(nullptr);
+					ParentComponentReferences.Add(FDMXPixelMappingComponentReference(StaticCastSharedRef<FDMXPixelMappingToolkit>(AsShared()), ParentComponent));
 				}
 			}
 		}
-
-		// Select the Parent Component 
-		SelectComponents(ParentComponentReferences);
-
-		UpdateBlueprintNodes(DMXPixelMapping);
 	}
+	
+	// Select the Parent Components
+	SelectComponents(ParentComponentReferences);
+
+	UpdateBlueprintNodes(DMXPixelMapping);
 }	
 
 bool FDMXPixelMappingToolkit::CanSizeSelectedComponentToTexture() const
@@ -942,6 +891,28 @@ void FDMXPixelMappingToolkit::SetupCommands()
 			})
 	);
 
+	GetToolkitCommands()->MapAction(
+		FDMXPixelMappingEditorCommands::Get().ToggleShowMatrixCells,
+		FExecuteAction::CreateLambda([]
+			{
+				if (UDMXPixelMappingLayoutSettings* LayoutSettings = GetMutableDefault<UDMXPixelMappingLayoutSettings>())
+				{
+					LayoutSettings->PreEditChange(nullptr);
+					LayoutSettings->bShowMatrixCells = !LayoutSettings->bShowMatrixCells;
+					LayoutSettings->PostEditChange();
+				}
+			}),
+		FCanExecuteAction(),
+				FIsActionChecked::CreateLambda([]()
+					{
+						if (const UDMXPixelMappingLayoutSettings* LayoutSettings = GetDefault<UDMXPixelMappingLayoutSettings>())
+						{
+							return LayoutSettings->bShowMatrixCells;
+						}
+						return false;
+					})
+				);
+
 	GetToolkitCommands()->MapAction(
 		FDMXPixelMappingEditorCommands::Get().ToggleShowComponentNames,
 		FExecuteAction::CreateLambda([]
diff --git a/Engine/Plugins/VirtualProduction/DMX/DMXPixelMapping/Source/DMXPixelMappingEditor/Private/ViewModels/DMXPixelMappingOutputComponentModel.cpp b/Engine/Plugins/VirtualProduction/DMX/DMXPixelMapping/Source/DMXPixelMappingEditor/Private/ViewModels/DMXPixelMappingOutputComponentModel.cpp
--- a/Engine/Plugins/VirtualProduction/DMX/DMXPixelMapping/Source/DMXPixelMappingEditor/Private/ViewModels/DMXPixelMappingOutputComponentModel.cpp
+++ b/Engine/Plugins/VirtualProduction/DMX/DMXPixelMapping/Source/DMXPixelMappingEditor/Private/ViewModels/DMXPixelMappingOutputComponentModel.cpp
@@ -3,6 +3,7 @@
 #include "ViewModels/DMXPixelMappingOutputComponentModel.h"
 
 #include "Components/DMXPixelMappingFixtureGroupComponent.h"
+#include "DMXPixelMappingLayoutSettings.h"
 #include "DMXPixelMappingTypes.h"
 #include "Components/DMXPixelMappingFixtureGroupItemComponent.h"
 #include "Components/DMXPixelMappingMatrixCellComponent.h"
@@ -70,15 +71,27 @@ FText FDMXPixelMappingOutputComponentModel::GetName() const
 	return FText::GetEmpty();
 }
 
-bool FDMXPixelMappingOutputComponentModel::ShouldDrawName() const
+bool FDMXPixelMappingOutputComponentModel::ShouldDraw() const
 {
-	if (UDMXPixelMappingOutputComponent* OutputComponent = Cast<UDMXPixelMappingOutputComponent>(WeakOutputComponent.Get()))
+	if (UDMXPixelMappingMatrixCellComponent* CellComponent = Cast<UDMXPixelMappingMatrixCellComponent>(WeakOutputComponent.Get()))
 	{
-		return OutputComponent->GetClass() != UDMXPixelMappingMatrixCellComponent::StaticClass();
+		const UDMXPixelMappingLayoutSettings* LayoutSettings = GetDefault<UDMXPixelMappingLayoutSettings>();
+		return LayoutSettings->bShowMatrixCells;
 	}
+
 	return true;
 }
 
+bool FDMXPixelMappingOutputComponentModel::ShouldDrawName() const
+{
+	const UDMXPixelMappingLayoutSettings* LayoutSettings = GetDefault<UDMXPixelMappingLayoutSettings>();
+	check(LayoutSettings);
+
+	return LayoutSettings->bShowComponentNames &&
+		WeakOutputComponent.IsValid() && 
+		WeakOutputComponent->GetClass() != UDMXPixelMappingMatrixCellComponent::StaticClass();
+}
+
 bool FDMXPixelMappingOutputComponentModel::ShouldDrawNameAbove() const
 {
 	if (UDMXPixelMappingOutputComponent* OutputComponent = Cast<UDMXPixelMappingOutputComponent>(WeakOutputComponent.Get()))
@@ -90,9 +103,24 @@ bool FDMXPixelMappingOutputComponentModel::ShouldDrawNameAbove() const
 	return true;
 }
 
-bool FDMXPixelMappingOutputComponentModel::HasPatchInfo() const
+bool FDMXPixelMappingOutputComponentModel::ShouldDrawCellID() const
 {
-	return GetFixturePatch() != nullptr;
+	const UDMXPixelMappingLayoutSettings* LayoutSettings = GetDefault<UDMXPixelMappingLayoutSettings>();
+	check(LayoutSettings);
+
+	return
+		LayoutSettings->bShowCellIDs &&
+		WeakOutputComponent->GetClass() == UDMXPixelMappingMatrixCellComponent::StaticClass();
+}
+
+bool FDMXPixelMappingOutputComponentModel::ShouldDrawPatchInfo() const
+{
+	const UDMXPixelMappingLayoutSettings* LayoutSettings = GetDefault<UDMXPixelMappingLayoutSettings>();
+	check(LayoutSettings);
+
+	return
+		LayoutSettings->bShowPatchInfo &&
+		GetFixturePatch() != nullptr;;
 }
 
 FText FDMXPixelMappingOutputComponentModel::GetAddressesText() const
@@ -107,7 +135,6 @@ FText FDMXPixelMappingOutputComponentModel::GetAddressesText() const
 	return FText::GetEmpty();
 }
 
-
 FText FDMXPixelMappingOutputComponentModel::GetFixtureIDText() const
 {
 	if (UDMXMVRFixtureNode* FixtureNode = WeakFixtureNode.Get())
@@ -119,15 +146,6 @@ FText FDMXPixelMappingOutputComponentModel::GetFixtureIDText() const
 	return FText::GetEmpty();
 }
 
-bool FDMXPixelMappingOutputComponentModel::HasCellID() const
-{
-	if (UDMXPixelMappingMatrixCellComponent* MatrixCell = Cast<UDMXPixelMappingMatrixCellComponent>(WeakOutputComponent.Get()))
-	{
-		return true;
-	}
-	return false;
-}
-
 FText FDMXPixelMappingOutputComponentModel::GetCellIDText() const
 {
 	if (UDMXPixelMappingMatrixCellComponent* MatrixCell = Cast<UDMXPixelMappingMatrixCellComponent>(WeakOutputComponent.Get()))
diff --git a/Engine/Plugins/VirtualProduction/DMX/DMXPixelMapping/Source/DMXPixelMappingEditor/Private/ViewModels/DMXPixelMappingOutputComponentModel.h b/Engine/Plugins/VirtualProduction/DMX/DMXPixelMapping/Source/DMXPixelMappingEditor/Private/ViewModels/DMXPixelMappingOutputComponentModel.h
--- a/Engine/Plugins/VirtualProduction/DMX/DMXPixelMapping/Source/DMXPixelMappingEditor/Private/ViewModels/DMXPixelMappingOutputComponentModel.h
+++ b/Engine/Plugins/VirtualProduction/DMX/DMXPixelMapping/Source/DMXPixelMappingEditor/Private/ViewModels/DMXPixelMappingOutputComponentModel.h
@@ -39,14 +39,20 @@ public:
 	/** Returns the Name of the component */
 	FText GetName() const;
 
+	/** Returns true if a widget ever should be drawn */
+	bool ShouldDraw() const;
+
 	/** Returns true if the name should drawn its name */
 	bool ShouldDrawName() const;
 
 	/** Returns true if the name should drawn its name above the widget */
 	bool ShouldDrawNameAbove() const;
 
-	/** Returns true if this component has a fixture patch and related info */
-	bool HasPatchInfo() const;
+	/** Returns true if the cell ID should be shown */
+	bool ShouldDrawCellID() const;
+
+	/** Returns true if the patch inifo should be shown */
+	bool ShouldDrawPatchInfo() const;
 
 	/** Returns the addresses of the patch as text */
 	FText GetAddressesText() const;
@@ -54,9 +60,6 @@ public:
 	/** Returns the Fixture ID as text. Empty unless a Matrix and Fixture Group Item component */
 	FText GetFixtureIDText() const;
 
-	/** Returns true if this component has a Cell ID */
-	bool HasCellID() const;
-
 	/** Returns the Cell ID as text. Empty unless a Matrix Cell component. */
 	FText GetCellIDText() const;
 
diff --git a/Engine/Plugins/VirtualProduction/DMX/DMXPixelMapping/Source/DMXPixelMappingEditor/Private/Viewports/DMXPixelMappingPreviewViewportClient.cpp b/Engine/Plugins/VirtualProduction/DMX/DMXPixelMapping/Source/DMXPixelMappingEditor/Private/Viewports/DMXPixelMappingPreviewViewportClient.cpp
--- a/Engine/Plugins/VirtualProduction/DMX/DMXPixelMapping/Source/DMXPixelMappingEditor/Private/Viewports/DMXPixelMappingPreviewViewportClient.cpp
+++ b/Engine/Plugins/VirtualProduction/DMX/DMXPixelMapping/Source/DMXPixelMappingEditor/Private/Viewports/DMXPixelMappingPreviewViewportClient.cpp
@@ -2,40 +2,146 @@
 
 #include "Viewports/DMXPixelMappingPreviewViewportClient.h"
 
-#include "Toolkits/DMXPixelMappingToolkit.h"
-#include "Components/DMXPixelMappingRendererComponent.h"
-#include "IDMXPixelMappingRenderer.h"
-
+#include "CanvasItem.h"
 #include "CanvasTypes.h"
+#include "Components/DMXPixelMappingRendererComponent.h"
 #include "Engine/TextureRenderTarget2D.h"
+#include "PixelMapRenderElement.h"
+#include "Slate/SceneViewport.h"
+#include "Toolkits/DMXPixelMappingToolkit.h"
 #include "UnrealClient.h"
+#include "Views/SDMXPixelMappingPreviewView.h"
+#include "Widgets/SDMXPixelMappingPreviewViewport.h"
+
 
-FDMXPixelMappingPreviewViewportClient::FDMXPixelMappingPreviewViewportClient(const TSharedPtr<FDMXPixelMappingToolkit>& InToolkit, TWeakPtr<SDMXPixelMappingPreviewViewport> InViewport)
-	: ToolkitWeakPtr(InToolkit)
-	, WeakViewport(InViewport)
+FDMXPixelMappingPreviewViewportClient::FDMXPixelMappingPreviewViewportClient(TWeakPtr<FDMXPixelMappingToolkit> InWeakToolkit, TWeakPtr<SDMXPixelMappingPreviewViewport> InWeakViewport)
+	: WeakToolkit(InWeakToolkit)
+	, WeakViewport(InWeakViewport)
+{}
+
+bool FDMXPixelMappingPreviewViewportClient::IsDrawingVisibleRectOnly() const
 {
-	check(InToolkit.IsValid() && WeakViewport.IsValid())
+	if (!WeakToolkit.IsValid() || !WeakViewport.IsValid())
+	{
+		return false;
+	}
+	const TSharedRef<FDMXPixelMappingToolkit> Toolkit = WeakToolkit.Pin().ToSharedRef();
+	const TSharedRef<SDMXPixelMappingPreviewView> PreviewView = Toolkit->GetOrCreatePreviewView();
+
+	UDMXPixelMappingRendererComponent* RendererComponent = Toolkit->GetActiveRendererComponent();
+	UTexture* InputTexture = RendererComponent ? RendererComponent->GetRenderedInputTexture() : nullptr;
+	if (!InputTexture)
+	{
+		return false;
+	}
+
+	const int32 ExcessThreshold = GMaxTextureDimensions / 4;
+
+	return
+		InputTexture->GetSurfaceWidth() * PreviewView->GetZoomAmount() > ExcessThreshold ||
+		InputTexture->GetSurfaceHeight() * PreviewView->GetZoomAmount() > ExcessThreshold;
+}
+
+FBox2D FDMXPixelMappingPreviewViewportClient::GetVisibleTextureBoxGraphSpace() const
+{
+	if (!WeakToolkit.IsValid())
+	{
+		return FBox2D();
+	}
+	const TSharedRef<FDMXPixelMappingToolkit> Toolkit = WeakToolkit.Pin().ToSharedRef();
+	const TSharedRef<SDMXPixelMappingPreviewView> PreviewView = Toolkit->GetOrCreatePreviewView();
+
+	UDMXPixelMappingRendererComponent* RendererComponent = Toolkit->GetActiveRendererComponent();
+	UTexture* InputTexture = RendererComponent ? RendererComponent->GetRenderedInputTexture() : nullptr;
+	if (!InputTexture || !InputTexture->GetResource())
+	{
+		return FBox2D();
+	}
+
+	// Compute texture box in graph space
+	const FVector2D TextureDimensions = FVector2D(InputTexture->GetSurfaceWidth(), InputTexture->GetSurfaceHeight());
+	const FBox2D TextureBox(FVector2D::ZeroVector, TextureDimensions);
+
+	// Compute graph box in graph space
+	const FGeometry& GraphGeometry = PreviewView->GetGraphTickSpaceGeometry();
+	const FVector2D GraphSize = GraphGeometry.GetAbsoluteSize() / PreviewView->GetZoomAmount();
+	const FVector2D ViewOffset = PreviewView->GetViewOffset();
+	const FBox2D GraphBox(ViewOffset, ViewOffset + GraphSize);
+
+	return TextureBox.Overlap(GraphBox);
 }
 
 void FDMXPixelMappingPreviewViewportClient::Draw(FViewport* InViewport, FCanvas* Canvas)
 {
-	Canvas->Clear(FColor::Transparent);
+	check(IsInGameThread());
+
+	Canvas->Clear(FColor::Black);
 
-	TSharedPtr<FDMXPixelMappingToolkit> Toolkit = ToolkitWeakPtr.Pin();
-	check(Toolkit.IsValid());
+	if (!WeakToolkit.IsValid())
+	{
+		return;
+	}
+	const TSharedRef<FDMXPixelMappingToolkit> Toolkit = WeakToolkit.Pin().ToSharedRef();
 
 	UDMXPixelMappingRendererComponent* RendererComponent = Toolkit->GetActiveRendererComponent();
-	if (RendererComponent)
+	UTexture* InputTexture = RendererComponent ? RendererComponent->GetRenderedInputTexture() : nullptr;
+	if (!InputTexture || !InputTexture->GetResource())
+	{
+		return;
+	}
+
+	const TSharedPtr<SDMXPixelMappingPreviewViewport> ViewportWidget = WeakViewport.Pin();
+	if (!ViewportWidget.IsValid())
 	{
-		for (UDMXPixelMappingOutputComponent* OutputComponent : Toolkit->GetActiveOutputComponents())
+		return;
+	}
+
+	const TSharedPtr<FSceneViewport> SceneViewport = ViewportWidget->GetViewport();
+	if (!SceneViewport.IsValid())
+	{
+		return;
+	}
+
+	// Draw the current render elements
+	const bool bDrawVisibleRectOnly = IsDrawingVisibleRectOnly();
+
+	using namespace UE::DMXPixelMapping::Rendering;
+	const TArray<TSharedRef<FPixelMapRenderElement>> RenderElements = RendererComponent->GetPixelMapRenderElements();
+	for (const TSharedRef<FPixelMapRenderElement>& Element : RenderElements)
+	{
+		const FVector2D UV = Element->GetParameters().UV;
+		const FVector2D UVSize = Element->GetParameters().UVSize;
+
+		if (bDrawVisibleRectOnly)
 		{
-			UTextureRenderTarget2D* InputTexture = RendererComponent->GetPreviewRenderTarget();
-			const TSharedPtr<IDMXPixelMappingRenderer>& Renderer = RendererComponent->GetRenderer();
+			const FVector2D PositionGraphSpace = UV * FVector2D{ InputTexture->GetSurfaceWidth(), InputTexture->GetSurfaceHeight() };
+			const FVector2D SizeGraphSpace = UVSize * FVector2D{ InputTexture->GetSurfaceWidth(), InputTexture->GetSurfaceHeight() };
+			
+			const FBox2D ComponentBox(PositionGraphSpace, PositionGraphSpace + SizeGraphSpace);
+			const FBox2D VisibleTextureBox = GetVisibleTextureBoxGraphSpace();
+			const FBox2D VisibleComponentBox = VisibleTextureBox.Overlap(ComponentBox);
 
-			if (Renderer.IsValid() && InputTexture != nullptr)
+			if (ComponentBox.bIsValid)
 			{
-				Renderer->RenderTextureToRectangle(InputTexture->GetResource(), InViewport->GetRenderTargetTexture(), InViewport->GetSizeXY(), InputTexture->SRGB);
+				const TSharedRef<SDMXPixelMappingPreviewView> PreviewView = Toolkit->GetOrCreatePreviewView();
+				const FVector2D Position = (VisibleComponentBox.Min - VisibleTextureBox.Min) * PreviewView->GetZoomAmount();
+				const FVector2D Size = VisibleComponentBox.GetSize() * PreviewView->GetZoomAmount();
+
+				FCanvasTileItem TileItem(Position, Size, Element->GetColor());
+				TileItem.BlendMode = ESimpleElementBlendMode::SE_BLEND_MAX;
+				Canvas->DrawItem(TileItem);
 			}
 		}
+		else
+		{
+			const FVector2D Position = UV * InViewport->GetSizeXY();
+			const FVector2D Size = UVSize * InViewport->GetSizeXY();
+
+			FCanvasTileItem TileItem(Position, Size, Element->GetColor());
+			TileItem.BlendMode = ESimpleElementBlendMode::SE_BLEND_MAX;
+			Canvas->DrawItem(TileItem);
+		}
 	}
+
+	Canvas->Flush_GameThread();
 }
diff --git a/Engine/Plugins/VirtualProduction/DMX/DMXPixelMapping/Source/DMXPixelMappingEditor/Private/Viewports/DMXPixelMappingPreviewViewportClient.h b/Engine/Plugins/VirtualProduction/DMX/DMXPixelMapping/Source/DMXPixelMappingEditor/Private/Viewports/DMXPixelMappingPreviewViewportClient.h
--- a/Engine/Plugins/VirtualProduction/DMX/DMXPixelMapping/Source/DMXPixelMappingEditor/Private/Viewports/DMXPixelMappingPreviewViewportClient.h
+++ b/Engine/Plugins/VirtualProduction/DMX/DMXPixelMapping/Source/DMXPixelMappingEditor/Private/Viewports/DMXPixelMappingPreviewViewportClient.h
@@ -4,21 +4,29 @@
 
 #include "ViewportClient.h"
 
-class SDMXPixelMappingPreviewViewport;
 class FDMXPixelMappingToolkit;
+class SDMXPixelMappingPreviewViewport;
+
 
 class FDMXPixelMappingPreviewViewportClient
 	: public FViewportClient
 {
 public:
 	/** Constructor */
-	FDMXPixelMappingPreviewViewportClient(const TSharedPtr<FDMXPixelMappingToolkit>& InToolkit, TWeakPtr<SDMXPixelMappingPreviewViewport> InViewport);
+	FDMXPixelMappingPreviewViewportClient(TWeakPtr<FDMXPixelMappingToolkit> InWeakToolkit, TWeakPtr<SDMXPixelMappingPreviewViewport> InWeakViewport);
+
+	/** If true, the widget draws the visible rect of the source texture */
+	bool IsDrawingVisibleRectOnly() const;
+
+	/** Returns the visible texture size in graph space */
+	FBox2D GetVisibleTextureBoxGraphSpace() const;
 
+protected:
 	//~ Begin FViewportClient Interface
 	virtual void Draw(FViewport* Viewport, FCanvas* Canvas) override;
 	//~ End FViewportClient Interface
 
 private:
-	TWeakPtr<FDMXPixelMappingToolkit> ToolkitWeakPtr;
+	TWeakPtr<FDMXPixelMappingToolkit> WeakToolkit;
 	TWeakPtr<SDMXPixelMappingPreviewViewport> WeakViewport;
 };
diff --git a/Engine/Plugins/VirtualProduction/DMX/DMXPixelMapping/Source/DMXPixelMappingEditor/Private/Viewports/DMXPixelMappingSourceTextureViewportClient.cpp b/Engine/Plugins/VirtualProduction/DMX/DMXPixelMapping/Source/DMXPixelMappingEditor/Private/Viewports/DMXPixelMappingSourceTextureViewportClient.cpp
--- a/Engine/Plugins/VirtualProduction/DMX/DMXPixelMapping/Source/DMXPixelMappingEditor/Private/Viewports/DMXPixelMappingSourceTextureViewportClient.cpp
+++ b/Engine/Plugins/VirtualProduction/DMX/DMXPixelMapping/Source/DMXPixelMappingEditor/Private/Viewports/DMXPixelMappingSourceTextureViewportClient.cpp
@@ -19,12 +19,11 @@ FDMXPixelMappingSourceTextureViewportClient::FDMXPixelMappingSourceTextureViewpo
 	, WeakSourceTextureViewport(InSourceTextureViewport)
 {}
 
-
-bool FDMXPixelMappingSourceTextureViewportClient::DrawOnlyVisibleRect() const
+bool FDMXPixelMappingSourceTextureViewportClient::IsDrawingVisibleRectOnly() const
 {
 	if (!WeakToolkit.IsValid())
 	{
-		return 0.0;
+		return false;
 	}
 	const TSharedRef<FDMXPixelMappingToolkit> Toolkit = WeakToolkit.Pin().ToSharedRef();
 	const TSharedRef<SDMXPixelMappingDesignerView> DesignerView = Toolkit->GetOrCreateDesignerView();
@@ -105,7 +104,7 @@ void FDMXPixelMappingSourceTextureViewportClient::Draw(FViewport* InViewport, FC
 	const FVector2D UV0 = VisibleRectBox.Min / TextureBox.GetSize();
 	const FVector2D UV1 = VisibleRectBox.Max / TextureBox.GetSize();
 
-	if (DrawOnlyVisibleRect())
+	if (IsDrawingVisibleRectOnly())
 	{
 		const TSharedPtr<FSceneViewport> SceneViewport = ViewportWidget->GetViewport();
 		if (!SceneViewport.IsValid())
diff --git a/Engine/Plugins/VirtualProduction/DMX/DMXPixelMapping/Source/DMXPixelMappingEditor/Private/Viewports/DMXPixelMappingSourceTextureViewportClient.h b/Engine/Plugins/VirtualProduction/DMX/DMXPixelMapping/Source/DMXPixelMappingEditor/Private/Viewports/DMXPixelMappingSourceTextureViewportClient.h
--- a/Engine/Plugins/VirtualProduction/DMX/DMXPixelMapping/Source/DMXPixelMappingEditor/Private/Viewports/DMXPixelMappingSourceTextureViewportClient.h
+++ b/Engine/Plugins/VirtualProduction/DMX/DMXPixelMapping/Source/DMXPixelMappingEditor/Private/Viewports/DMXPixelMappingSourceTextureViewportClient.h
@@ -15,11 +15,8 @@ public:
 	/** Constructor */
 	FDMXPixelMappingSourceTextureViewportClient(const TSharedPtr<FDMXPixelMappingToolkit>& Toolkit, TWeakPtr<SDMXPixelMappingSourceTextureViewport> InViewport);
 	
-	/**
-	 * Returns true if it only the visible rectangle is drawn, instead of the whole texture.
-	 * This is favorable when zoomed in closely and the viewport size gets close to, or exceeds GMaxTextureDimensions.
-	 */
-	bool DrawOnlyVisibleRect() const;
+	/** If true, the widget draws the visible rect of the source texture */
+	bool IsDrawingVisibleRectOnly() const;
 
 	/** Returns the visible texture size in graph space */
 	FBox2D GetVisibleTextureBoxGraphSpace() const;
diff --git a/Engine/Plugins/VirtualProduction/DMX/DMXPixelMapping/Source/DMXPixelMappingEditor/Private/Views/SDMXPixelMappingDesignerView.cpp b/Engine/Plugins/VirtualProduction/DMX/DMXPixelMapping/Source/DMXPixelMappingEditor/Private/Views/SDMXPixelMappingDesignerView.cpp
--- a/Engine/Plugins/VirtualProduction/DMX/DMXPixelMapping/Source/DMXPixelMappingEditor/Private/Views/SDMXPixelMappingDesignerView.cpp
+++ b/Engine/Plugins/VirtualProduction/DMX/DMXPixelMapping/Source/DMXPixelMappingEditor/Private/Views/SDMXPixelMappingDesignerView.cpp
@@ -316,7 +316,7 @@ FReply SDMXPixelMappingDesignerView::OnMouseButtonUp(const FGeometry& MyGeometry
 	// Select the Renderer if no Output Component was under the mouse
 	if (!PendingSelectedComponent.IsValid())
 	{
-		if (TSharedPtr<FDMXPixelMappingToolkit> ToolkitPtr = ToolkitWeakPtr.Pin())
+		if (TSharedPtr<FDMXPixelMappingToolkit> ToolkitPtr = WeakToolkit.Pin())
 		{
 			if (UDMXPixelMappingRendererComponent* RendererComponent = ToolkitPtr->GetActiveRendererComponent())
 			{
@@ -374,7 +374,7 @@ FReply SDMXPixelMappingDesignerView::OnKeyDown(const FGeometry& MyGeometry, cons
 {
 	if (InKeyEvent.GetKey() == EKeys::Delete)
 	{
-		if (const TSharedPtr<FDMXPixelMappingToolkit> ToolkitPtr = ToolkitWeakPtr.Pin())
+		if (const TSharedPtr<FDMXPixelMappingToolkit> ToolkitPtr = WeakToolkit.Pin())
 		{
 			const FScopedTransaction Transaction(FText::Format(LOCTEXT("DMXPixelMapping.RemoveComponents", "Remove {0}|plural(one=Component, other=Components)"), ToolkitPtr->GetSelectedComponents().Num()));
 
@@ -391,7 +391,7 @@ void SDMXPixelMappingDesignerView::Tick(const FGeometry& AllottedGeometry, const
 {
 	SDMXPixelMappingSurface::Tick(AllottedGeometry, InCurrentTime, InDeltaTime);
 
-	if (const TSharedPtr<FDMXPixelMappingToolkit> Toolkit = ToolkitWeakPtr.Pin())
+	if (const TSharedPtr<FDMXPixelMappingToolkit> Toolkit = WeakToolkit.Pin())
 	{
 		// Handle drag over on tick
 		if (PendingDragDropOp.IsValid())
@@ -545,7 +545,7 @@ void SDMXPixelMappingDesignerView::OnDragLeave(const FDragDropEvent& DragDropEve
 	const TSharedPtr<FDMXPixelMappingDragDropOp> DragDropOp = DragDropEvent.GetOperationAs<FDMXPixelMappingDragDropOp>();
 	if (DragDropOp.IsValid())
 	{
-		if (TSharedPtr<FDMXPixelMappingToolkit> ToolkitPtr = ToolkitWeakPtr.Pin())
+		if (TSharedPtr<FDMXPixelMappingToolkit> ToolkitPtr = WeakToolkit.Pin())
 		{
 			FScopedRestoreSelection(ToolkitPtr.ToSharedRef(), StaticCastSharedRef<SDMXPixelMappingDesignerView>(AsShared()));
 
@@ -595,27 +595,18 @@ FReply SDMXPixelMappingDesignerView::OnDrop(const FGeometry& MyGeometry, const F
 
 FSlateRect SDMXPixelMappingDesignerView::ComputeAreaBounds() const
 {
-	FSlateRect Bounds = FSlateRect(0.f, 0.f, 0.f, 0.f);
-
-	const TSet<FDMXPixelMappingComponentReference> SelectedComponentReferences = GetSelectedComponents();
-	for (const FDMXPixelMappingComponentReference& ComponentReference : SelectedComponentReferences)
+	if (SourceTextureViewport.IsValid())
 	{
-		if (UDMXPixelMappingBaseComponent* Component = ComponentReference.GetComponent())
+		const FVector2D Size
 		{
-			if (UDMXPixelMappingRendererComponent* RendererComponent = Component->GetRendererComponent())
-			{
-				Bounds.Left = FMath::Min(RendererComponent->GetPosition().X, Bounds.Left);
-				Bounds.Top = FMath::Min(RendererComponent->GetPosition().Y, Bounds.Top);
-
-				Bounds.Right = FMath::Max(RendererComponent->GetPosition().X + RendererComponent->GetSize().X, Bounds.Right);
-				Bounds.Bottom = FMath::Max(RendererComponent->GetPosition().Y + RendererComponent->GetSize().Y, Bounds.Bottom);
+			FMath::Max(SourceTextureViewport->GetWidthGraphSpace().Get(), 1.0),
+			FMath::Max(SourceTextureViewport->GetHeightGraphSpace().Get(), 1.0),
+		};
 
-				break;
-			}
-		}
+		return FSlateRect(0.f, 0.f, Size.X, Size.Y);
 	}
 
-	return Bounds;
+	return FSlateRect();
 }
 
 int32 SDMXPixelMappingDesignerView::GetGraphRulePeriod() const
@@ -645,7 +636,7 @@ void SDMXPixelMappingDesignerView::PostRedo(bool bSuccess)
 
 void SDMXPixelMappingDesignerView::RebuildDesigner()
 {
-	const TSharedPtr<FDMXPixelMappingToolkit> Toolkit = ToolkitWeakPtr.Pin();
+	const TSharedPtr<FDMXPixelMappingToolkit> Toolkit = WeakToolkit.Pin();
 	if (!Toolkit.IsValid())
 	{
 		return;
@@ -666,6 +657,7 @@ void SDMXPixelMappingDesignerView::RebuildDesigner()
 	{
 		const TSharedRef<SConstraintCanvas> ComponentCanvas = SNew(SConstraintCanvas);
 		DesignCanvas->AddSlot()
+			.Alignment(FVector2D::ZeroVector)
 			[
 				ComponentCanvas
 			];
@@ -709,7 +701,7 @@ void SDMXPixelMappingDesignerView::CreateExtensionWidgetsForSelection()
 		!OutputComponent->IsLockInDesigner())
 	{
 		// Add transform handles
-		constexpr float Offset = 10.f;
+		constexpr float Offset = 0.1f;
 		const TSharedPtr<SDMXPixelMappingDesignerView> Self = SharedThis(this);
 		TransformHandles.Add(SNew(SDMXPixelMappingTransformHandle, Self, EDMXPixelMappingTransformDirection::CenterRight, FVector2D(Offset, 0.f)));
 		TransformHandles.Add(SNew(SDMXPixelMappingTransformHandle, Self, EDMXPixelMappingTransformDirection::BottomCenter, FVector2D(0.f, Offset)));
@@ -883,7 +875,7 @@ void SDMXPixelMappingDesignerView::OnSelectedComponentsChanged()
 
 void SDMXPixelMappingDesignerView::ResolvePendingSelectedComponents(bool bClearPreviousSelection)
 {
-	const TSharedPtr<FDMXPixelMappingToolkit> Toolkit = ToolkitWeakPtr.Pin();
+	const TSharedPtr<FDMXPixelMappingToolkit> Toolkit = WeakToolkit.Pin();
 	const bool bHasPendingSelectedComponent = PendingSelectedComponent.IsValid();
 	if (Toolkit.IsValid() && bHasPendingSelectedComponent)
 	{
@@ -913,7 +905,7 @@ void SDMXPixelMappingDesignerView::ResolvePendingSelectedComponents(bool bClearP
 
 		if (!bClearPreviousSelection)
 		{
-			SelectedComponents.Append(ToolkitWeakPtr.Pin()->GetSelectedComponents());
+			SelectedComponents.Append(WeakToolkit.Pin()->GetSelectedComponents());
 		}
 		Toolkit->SelectComponents(SelectedComponents);
 
@@ -960,8 +952,6 @@ UDMXPixelMappingOutputComponent* SDMXPixelMappingDesignerView::GetComponentUnder
 		FVector2D GraphSpaceCursorPosition;
 		if (GetGraphSpaceCursorPosition(GraphSpaceCursorPosition))
 		{
-			// Widgets are offset by 0.5f to reside over the pixel
-			GraphSpaceCursorPosition += FVector2D(.5f, .5f);
 			RendererComponent->ForEachChild([&ComponentUnderCursor, &GraphSpaceCursorPosition](UDMXPixelMappingBaseComponent* InComponent)
 				{
 					if (UDMXPixelMappingOutputComponent* OutputComponent = Cast<UDMXPixelMappingOutputComponent>(InComponent))
@@ -985,7 +975,7 @@ UDMXPixelMappingOutputComponent* SDMXPixelMappingDesignerView::GetComponentUnder
 
 TSet<FDMXPixelMappingComponentReference> SDMXPixelMappingDesignerView::GetSelectedComponents() const
 {
-	if (TSharedPtr<FDMXPixelMappingToolkit> ToolkitPtr = ToolkitWeakPtr.Pin())
+	if (TSharedPtr<FDMXPixelMappingToolkit> ToolkitPtr = WeakToolkit.Pin())
 	{
 		return ToolkitPtr->GetSelectedComponents();
 	}
@@ -1076,7 +1066,7 @@ FGeometry SDMXPixelMappingDesignerView::MakeGeometryWindowLocal(const FGeometry&
 
 void SDMXPixelMappingDesignerView::HandleDragEnterFromDetailsOrPalette(const TSharedPtr<FDMXPixelMappingDragDropOp>& TemplateDragDropOp)
 {
-	if (TSharedPtr<FDMXPixelMappingToolkit> ToolkitPtr = ToolkitWeakPtr.Pin())
+	if (TSharedPtr<FDMXPixelMappingToolkit> ToolkitPtr = WeakToolkit.Pin())
 	{
 		if (TemplateDragDropOp.IsValid())
 		{
diff --git a/Engine/Plugins/VirtualProduction/DMX/DMXPixelMapping/Source/DMXPixelMappingEditor/Private/Views/SDMXPixelMappingDetailsView.cpp b/Engine/Plugins/VirtualProduction/DMX/DMXPixelMapping/Source/DMXPixelMappingEditor/Private/Views/SDMXPixelMappingDetailsView.cpp
--- a/Engine/Plugins/VirtualProduction/DMX/DMXPixelMapping/Source/DMXPixelMappingEditor/Private/Views/SDMXPixelMappingDetailsView.cpp
+++ b/Engine/Plugins/VirtualProduction/DMX/DMXPixelMapping/Source/DMXPixelMappingEditor/Private/Views/SDMXPixelMappingDetailsView.cpp
@@ -16,7 +16,6 @@
 #include "Customizations/DMXPixelMappingDetailCustomization_Screen.h"
 #include "Customizations/DMXPixelMappingDetailCustomization_Renderer.h"
 #include "Customizations/DMXPixelMappingDetailCustomization_Matrix.h"
-#include "Customizations/DMXPixelMappingDetailCustomization_MatrixCell.h"
 #include "DetailsViewArgs.h"
 #include "DMXPixelMappingComponentReference.h"
 #include "IDetailsView.h"
@@ -58,21 +57,6 @@ void SDMXPixelMappingDetailsView::Construct(const FArguments& InArgs, const TSha
 	OnSelectedComponentsChanged();
 }
 
-SDMXPixelMappingDetailsView::~SDMXPixelMappingDetailsView()
-{
-	if (PropertyView.IsValid())
-	{
-		PropertyView->UnregisterInstancedCustomPropertyLayout(UDMXPixelMappingFixtureGroupComponent::StaticClass());
-		PropertyView->UnregisterInstancedCustomPropertyLayout(UDMXPixelMappingFixtureGroupItemComponent::StaticClass());
-		PropertyView->UnregisterInstancedCustomPropertyLayout(UDMXPixelMappingScreenComponent::StaticClass());
-		PropertyView->UnregisterInstancedCustomPropertyLayout(UDMXPixelMappingRendererComponent::StaticClass());
-		PropertyView->UnregisterInstancedCustomPropertyLayout(UDMXPixelMappingMatrixComponent::StaticClass());
-		PropertyView->UnregisterInstancedCustomPropertyLayout(UDMXPixelMappingMatrixCellComponent::StaticClass());
-		
-		PropertyView->UnregisterInstancedCustomPropertyLayout(UDMXPixelMappingColorSpace_RGBCMY::StaticClass());
-	}
-}
-
 void SDMXPixelMappingDetailsView::OnSelectedComponentsChanged()
 {
 	// Clear selection in the property view.
@@ -130,9 +114,6 @@ void SDMXPixelMappingDetailsView::RegisterCustomizations()
 	FOnGetDetailCustomizationInstance MatrixCustomizationInstance = FOnGetDetailCustomizationInstance::CreateStatic(&FDMXPixelMappingDetailCustomization_Matrix::MakeInstance, ToolkitWeakPtr);
 	PropertyView->RegisterInstancedCustomPropertyLayout(UDMXPixelMappingMatrixComponent::StaticClass(), MatrixCustomizationInstance);
 
-	FOnGetDetailCustomizationInstance MatrixCellCustomizationInstance = FOnGetDetailCustomizationInstance::CreateStatic(&FDMXPixelMappingDetailCustomization_MatrixCell::MakeInstance, ToolkitWeakPtr);
-	PropertyView->RegisterInstancedCustomPropertyLayout(UDMXPixelMappingMatrixCellComponent::StaticClass(), MatrixCellCustomizationInstance);
-
 	FOnGetDetailCustomizationInstance ColorSpaceCustomizationInstance_RGBCMY = FOnGetDetailCustomizationInstance::CreateStatic(&FDMXPixelMappingColorSpaceDetails_RGBCMY::MakeInstance);
 	PropertyView->RegisterInstancedCustomPropertyLayout(UDMXPixelMappingColorSpace_RGBCMY::StaticClass(), ColorSpaceCustomizationInstance_RGBCMY);
 }
diff --git a/Engine/Plugins/VirtualProduction/DMX/DMXPixelMapping/Source/DMXPixelMappingEditor/Private/Views/SDMXPixelMappingDetailsView.h b/Engine/Plugins/VirtualProduction/DMX/DMXPixelMapping/Source/DMXPixelMappingEditor/Private/Views/SDMXPixelMappingDetailsView.h
--- a/Engine/Plugins/VirtualProduction/DMX/DMXPixelMapping/Source/DMXPixelMappingEditor/Private/Views/SDMXPixelMappingDetailsView.h
+++ b/Engine/Plugins/VirtualProduction/DMX/DMXPixelMapping/Source/DMXPixelMappingEditor/Private/Views/SDMXPixelMappingDetailsView.h
@@ -23,15 +23,12 @@ public:
 	 */
 	void Construct(const FArguments& InArgs, const TSharedPtr<FDMXPixelMappingToolkit>& InToolkit);
 
-	virtual ~SDMXPixelMappingDetailsView();
-
 private:
 	void OnSelectedComponentsChanged();
 
 	/** Registers the designer specific customizations */
 	void RegisterCustomizations();
 
-private:
 	TWeakPtr<FDMXPixelMappingToolkit> ToolkitWeakPtr;
 
 	/** Property viewing widget */
diff --git a/Engine/Plugins/VirtualProduction/DMX/DMXPixelMapping/Source/DMXPixelMappingEditor/Private/Views/SDMXPixelMappingPreviewView.cpp b/Engine/Plugins/VirtualProduction/DMX/DMXPixelMapping/Source/DMXPixelMappingEditor/Private/Views/SDMXPixelMappingPreviewView.cpp
--- a/Engine/Plugins/VirtualProduction/DMX/DMXPixelMapping/Source/DMXPixelMappingEditor/Private/Views/SDMXPixelMappingPreviewView.cpp
+++ b/Engine/Plugins/VirtualProduction/DMX/DMXPixelMapping/Source/DMXPixelMappingEditor/Private/Views/SDMXPixelMappingPreviewView.cpp
@@ -20,7 +20,7 @@
 
 void SDMXPixelMappingPreviewView::Construct(const FArguments& InArgs, const TSharedPtr<FDMXPixelMappingToolkit>& InToolkit)
 {
-	ToolkitWeakPtr = InToolkit;
+	WeakToolkit = InToolkit;
 
 	SDMXPixelMappingSurface::Construct(SDMXPixelMappingSurface::FArguments()
 		.AllowContinousZoomInterpolation(false)
@@ -76,7 +76,7 @@ void SDMXPixelMappingPreviewView::Construct(const FArguments& InArgs, const TSha
 				.HAlign(HAlign_Fill)
 				.VAlign(VAlign_Fill)
 				[
-					SNew(SDMXPixelMappingZoomPan)
+					SAssignNew(ZoomPan, SDMXPixelMappingZoomPan)
 					.ZoomAmount(this, &SDMXPixelMappingPreviewView::GetZoomAmount)
 					.ViewOffset(this, &SDMXPixelMappingPreviewView::GetViewOffset)
 					.Visibility(this, &SDMXPixelMappingPreviewView::IsZoomPanVisible)
@@ -107,6 +107,11 @@ void SDMXPixelMappingPreviewView::Construct(const FArguments& InArgs, const TSha
 	, InToolkit);
 }
 
+const FGeometry& SDMXPixelMappingPreviewView::GetGraphTickSpaceGeometry() const
+{
+	return ZoomPan->GetTickSpaceGeometry();
+}
+
 FReply SDMXPixelMappingPreviewView::OnMouseButtonDown(const FGeometry& MyGeometry, const FPointerEvent& MouseEvent)
 {
 	SDMXPixelMappingSurface::OnMouseButtonDown(MyGeometry, MouseEvent);
@@ -218,31 +223,33 @@ int32 SDMXPixelMappingPreviewView::GetSnapGridSize() const
 
 FSlateRect SDMXPixelMappingPreviewView::ComputeAreaBounds() const
 {
-	return FSlateRect(0, 0, GetPreviewAreaWidth().Get(), GetPreviewAreaHeight().Get());
-}
-
-FOptionalSize SDMXPixelMappingPreviewView::GetPreviewAreaWidth() const
-{
-	FVector2D Area, Size;
-	GetPreviewAreaAndSize(Area, Size);
+	const FVector2D Size
+	{
+		FMath::Max(GetWidthGraphSpace().Get(), 1.0),
+		FMath::Max(GetHeightGraphSpace().Get(), 1.0),
+	};
 
-	return Area.X;
+	return FSlateRect(0, 0, Size.X, Size.Y);
 }
 
-FOptionalSize SDMXPixelMappingPreviewView::GetPreviewAreaHeight() const
+FOptionalSize SDMXPixelMappingPreviewView::GetWidthGraphSpace() const
 {
-	FVector2D Area, Size;
-	GetPreviewAreaAndSize(Area, Size);
+	if (PreviewViewport.IsValid())
+	{
+		return PreviewViewport->GetWidthGraphSpace();
+	}
 
-	return Area.Y;
+	return FOptionalSize();
 }
 
-void SDMXPixelMappingPreviewView::GetPreviewAreaAndSize(FVector2D& Area, FVector2D& Size) const
+FOptionalSize SDMXPixelMappingPreviewView::GetHeightGraphSpace() const
 {
-	check(PreviewViewport.IsValid())
+	if (PreviewViewport.IsValid())
+	{
+		return PreviewViewport->GetHeightGraphSpace();
+	}
 
-	Area = FVector2D(PreviewViewport->GetPreviewAreaWidth().Get(), PreviewViewport->GetPreviewAreaHeight().Get());
-	Size = Area;
+	return FOptionalSize();
 }
 
 TSharedRef<SWidget> SDMXPixelMappingPreviewView::CreateOverlayUI()
@@ -296,7 +303,7 @@ TSharedRef<SWidget> SDMXPixelMappingPreviewView::CreateOverlayUI()
 
 EVisibility SDMXPixelMappingPreviewView::IsZoomPanVisible() const
 {
-	if (TSharedPtr<FDMXPixelMappingToolkit> Toolkit = ToolkitWeakPtr.Pin())
+	if (TSharedPtr<FDMXPixelMappingToolkit> Toolkit = WeakToolkit.Pin())
 	{
 		if (UDMXPixelMappingRendererComponent* RendererComponent = Toolkit->GetActiveRendererComponent())
 		{
@@ -357,9 +364,9 @@ FReply SDMXPixelMappingPreviewView::HandleZoomToFitClicked()
 
 const TSet<FDMXPixelMappingComponentReference>& SDMXPixelMappingPreviewView::GetSelectedComponents() const
 {
-	check(ToolkitWeakPtr.Pin().IsValid());
+	check(WeakToolkit.Pin().IsValid());
 
-	return ToolkitWeakPtr.Pin()->GetSelectedComponents();
+	return WeakToolkit.Pin()->GetSelectedComponents();
 }
 
 FDMXPixelMappingComponentReference SDMXPixelMappingPreviewView::GetSelectedComponent() const
diff --git a/Engine/Plugins/VirtualProduction/DMX/DMXPixelMapping/Source/DMXPixelMappingEditor/Private/Views/SDMXPixelMappingPreviewView.h b/Engine/Plugins/VirtualProduction/DMX/DMXPixelMapping/Source/DMXPixelMappingEditor/Private/Views/SDMXPixelMappingPreviewView.h
--- a/Engine/Plugins/VirtualProduction/DMX/DMXPixelMapping/Source/DMXPixelMappingEditor/Private/Views/SDMXPixelMappingPreviewView.h
+++ b/Engine/Plugins/VirtualProduction/DMX/DMXPixelMapping/Source/DMXPixelMappingEditor/Private/Views/SDMXPixelMappingPreviewView.h
@@ -4,14 +4,15 @@
 
 #include "Widgets/SDMXPixelMappingSurface.h"
 
-
 class FDMXPixelMappingToolkit;
 class SDMXPixelMappingPreviewViewport;
 struct FOptionalSize;
 class SBox;
+class SDMXPixelMappingZoomPan;
 class SOverlay;
 class FDMXPixelMappingComponentReference;
 
+
 class SDMXPixelMappingPreviewView
 	: public SDMXPixelMappingSurface
 {
@@ -28,6 +29,10 @@ public:
 	 */
 	void Construct(const FArguments& InArgs, const TSharedPtr<FDMXPixelMappingToolkit>& InToolkit);
 
+	/** Returns the geometry of the graph */
+	const FGeometry& GetGraphTickSpaceGeometry() const;
+
+protected:
 	// SWidget interface
 	virtual FReply OnMouseButtonDown(const FGeometry& MyGeometry, const FPointerEvent& MouseEvent) override;
 	virtual FReply OnMouseButtonUp(const FGeometry& MyGeometry, const FPointerEvent& MouseEvent) override;
@@ -51,11 +56,11 @@ public:
 	//~ End SDMXPixelMappingSurface interface
 
 private:
-	FOptionalSize GetPreviewAreaWidth() const;
+	/** Gets the height in graph space */
+	FOptionalSize GetHeightGraphSpace() const;
 
-	FOptionalSize GetPreviewAreaHeight() const;
-
-	void GetPreviewAreaAndSize(FVector2D& Area, FVector2D& Size) const;
+	/** Gets the width in graph space */
+	FOptionalSize GetWidthGraphSpace() const;
 
 	TSharedRef<SWidget> CreateOverlayUI();
 
@@ -77,6 +82,9 @@ private:
 	FText GetSelectedComponentParentNameText() const;
 	EVisibility GetTitleBarVisibility() const;
 
+	/** Zoom pan widget */
+	TSharedPtr<SDMXPixelMappingZoomPan> ZoomPan;
+
 	TSharedPtr<SDMXPixelMappingPreviewViewport> PreviewViewport;
 
 	/** Cache last mouse position to be used as a paste drop location */
diff --git a/Engine/Plugins/VirtualProduction/DMX/DMXPixelMapping/Source/DMXPixelMappingEditor/Private/Widgets/SDMXPixelMappingOutputComponent.cpp b/Engine/Plugins/VirtualProduction/DMX/DMXPixelMapping/Source/DMXPixelMappingEditor/Private/Widgets/SDMXPixelMappingOutputComponent.cpp
--- a/Engine/Plugins/VirtualProduction/DMX/DMXPixelMapping/Source/DMXPixelMappingEditor/Private/Widgets/SDMXPixelMappingOutputComponent.cpp
+++ b/Engine/Plugins/VirtualProduction/DMX/DMXPixelMapping/Source/DMXPixelMappingEditor/Private/Widgets/SDMXPixelMappingOutputComponent.cpp
@@ -3,12 +3,12 @@
 #include "Widgets/SDMXPixelMappingOutputComponent.h"
 
 #include "DMXPixelMappingLayoutSettings.h"
+#include "DMXPixelMappingEditorStyle.h"
 #include "DMXPixelMappingTypes.h"
 #include "DMXPixelMappingUtils.h"
 #include "DMXRuntimeUtils.h"
-#include "ViewModels/DMXPixelMappingOutputComponentModel.h"
 #include "Toolkits/DMXPixelMappingToolkit.h"
-
+#include "ViewModels/DMXPixelMappingOutputComponentModel.h"
 #include "Widgets/Layout/SBorder.h"
 #include "Widgets/Layout/SBox.h"
 #include "Widgets/Layout/SScaleBox.h"
@@ -35,10 +35,8 @@ void IDMXPixelMappingOutputComponentWidgetInterface::AddToCanvas(const TSharedRe
 		.Alignment(FVector2D::ZeroVector)
 		.Offset_Lambda([this]()
 			{
-				const FVector2D Postition = GetPosition();
-
-				// In the center of the top left pixel
-				return FMargin(Postition.X + .5f, Postition.Y + .5f);
+				const FVector2D Position = GetPosition();
+				return FMargin(Position.X, Position.Y);
 			})
 		.Expose(Slot)
 		[
@@ -81,67 +79,10 @@ void SDMXPixelMappingOutputComponent::Construct(const FArguments& InArgs, const
 			{
 				return Model->GetSize().Y;
 			})
-		[
-			SNew(SBorder)
-			.BorderImage_Lambda([this]()
-				{
-					BorderBrush.TintColor = Model->GetColor();
-					return &BorderBrush;
-				})
-			[
-				SNew(SVerticalBox)
-
-				+ SVerticalBox::Slot()
-				.HAlign(HAlign_Fill)
-				.VAlign(VAlign_Bottom)
-				.AutoHeight()
-				.Padding(FMargin(0.f, -22.f, 0.f, 0.f))
-				[
-					SNew(SBox)
-					.HeightOverride(22.f)
-					.HAlign(HAlign_Left)
-					.VAlign(VAlign_Bottom)
-					[
-						SAssignNew(AboveContentBorder, SBorder)
-						.BorderImage(FAppStyle::GetBrush("NoBorder"))
-					]
-				]
-
-				+ SVerticalBox::Slot()
-				.HAlign(HAlign_Fill)
-				.VAlign(VAlign_Fill)
-				.FillHeight(1.f / 3.f)
-				.Padding(FMargin(-2.f, -2.f, -2.f, -2.f))
-				[
-					SAssignNew(TopContentBorder, SBorder)
-					.BorderImage(FAppStyle::GetBrush("NoBorder"))
-				]
-
-				+ SVerticalBox::Slot()
-				.HAlign(HAlign_Fill)
-				.VAlign(VAlign_Fill)
-				.FillHeight(1.f / 3.f)
-				.Padding(FMargin(-2.f, -2.f, -2.f, -2.f))
-				[
-					SAssignNew(MiddleContentBorder, SBorder)
-					.BorderImage(FAppStyle::GetBrush("NoBorder"))
-				]
-
-				+ SVerticalBox::Slot()
-				.HAlign(HAlign_Fill)
-				.VAlign(VAlign_Fill)
-				.FillHeight(1.f / 3.f)
-				.Padding(FMargin(-2.f, -2.f, -2.f, -2.f))
-				[
-					SAssignNew(BottomContentBorder, SBorder)
-					.BorderImage(FAppStyle::GetBrush("NoBorder"))
-				]
-			]
-		]
 	];
 
-	UpdateChildSlots();
-	UDMXPixelMappingLayoutSettings::GetOnLayoutSettingsChanged().AddSP(this, &SDMXPixelMappingOutputComponent::UpdateChildSlots);
+	ForceRefresh();
+	UDMXPixelMappingLayoutSettings::GetOnLayoutSettingsChanged().AddSP(this, &SDMXPixelMappingOutputComponent::RequestRefresh);
 }
 
 bool SDMXPixelMappingOutputComponent::Equals(UDMXPixelMappingBaseComponent* Component) const
@@ -158,20 +99,90 @@ FVector2D SDMXPixelMappingOutputComponent::GetPosition() const
 	return Model->GetPosition();
 }
 
-void SDMXPixelMappingOutputComponent::UpdateChildSlots()
+void SDMXPixelMappingOutputComponent::RequestRefresh()
 {
-	AboveContentBorder->ClearContent();
-	TopContentBorder->ClearContent();
-	MiddleContentBorder->ClearContent();
-	BottomContentBorder->ClearContent();
+	if (!RefreshTimerHandle.IsValid())
+	{
+		RefreshTimerHandle = GEditor->GetTimerManager()->SetTimerForNextTick(FTimerDelegate::CreateSP(this, &SDMXPixelMappingOutputComponent::ForceRefresh));
+	}
+}
 
-	const UDMXPixelMappingLayoutSettings* LayoutSettings = GetDefault<UDMXPixelMappingLayoutSettings>();
-	if (!LayoutSettings)
+void SDMXPixelMappingOutputComponent::ForceRefresh()
+{
+	RefreshTimerHandle.Invalidate();
+	if (Model->ShouldDraw())
 	{
-		return;
+		ComponentBox->SetContent(CreateContent());
 	}
+	else
+	{
+		ComponentBox->SetContent(SNullWidget::NullWidget);
+	}
+}
 
-	if (LayoutSettings->bShowComponentNames && Model->ShouldDrawName())
+TSharedRef<SWidget> SDMXPixelMappingOutputComponent::CreateContent()
+{
+	const TSharedRef<SBorder> Content =
+		SNew(SBorder)
+		.BorderImage(FDMXPixelMappingEditorStyle::Get().GetBrush("DMXPixelMappingEditor.ComponentBorder"))
+		.BorderBackgroundColor_Lambda([this]()
+			{
+				return Model->GetColor();
+			})
+		[
+			SNew(SVerticalBox)
+
+			+ SVerticalBox::Slot()
+			.HAlign(HAlign_Fill)
+			.VAlign(VAlign_Bottom)
+			.AutoHeight()
+			.Padding(FMargin(0.f, -22.f, 0.f, 0.f))
+			[
+				SNew(SBox)
+				.HeightOverride(22.f)
+				.HAlign(HAlign_Left)
+				.VAlign(VAlign_Bottom)
+				[
+					SAssignNew(AboveContentBorder, SBorder)
+					.BorderImage(FAppStyle::GetBrush("NoBorder"))
+				]
+			]
+
+			+ SVerticalBox::Slot()
+			.HAlign(HAlign_Fill)
+			.VAlign(VAlign_Fill)
+			.FillHeight(1.f / 3.f)
+			.Padding(FMargin(-2.f, -2.f, -2.f, -2.f))
+			[
+				SAssignNew(TopContentBorder, SBorder)
+				.BorderImage(FAppStyle::GetBrush("NoBorder"))
+			]
+
+			+ SVerticalBox::Slot()
+			.HAlign(HAlign_Fill)
+			.VAlign(VAlign_Fill)
+			.FillHeight(1.f / 3.f)
+			.Padding(FMargin(-2.f, -2.f, -2.f, -2.f))
+			[
+				SAssignNew(MiddleContentBorder, SBorder)
+				.BorderImage(FAppStyle::GetBrush("NoBorder"))
+			]
+
+			+ SVerticalBox::Slot()
+			.HAlign(HAlign_Fill)
+			.VAlign(VAlign_Fill)
+			.FillHeight(1.f / 3.f)
+			.Padding(FMargin(-2.f, -2.f, -2.f, -2.f))
+			[
+				SAssignNew(BottomContentBorder, SBorder)
+				.BorderImage(FAppStyle::GetBrush("NoBorder"))
+			]
+		];
+
+	const UDMXPixelMappingLayoutSettings* LayoutSettings = GetDefault<UDMXPixelMappingLayoutSettings>();
+	check(LayoutSettings);
+
+	if (Model->ShouldDrawName())
 	{
 		if (Model->ShouldDrawNameAbove())
 		{
@@ -183,15 +194,17 @@ void SDMXPixelMappingOutputComponent::UpdateChildSlots()
 		}
 	}
 
-	if (LayoutSettings->bShowCellIDs && Model->HasCellID())
+	if (Model->ShouldDrawCellID())
 	{
 		CreateCellIDChildSlot();
 	}
 
-	if (LayoutSettings->bShowPatchInfo && Model->HasPatchInfo())
+	if (Model->ShouldDrawPatchInfo())
 	{
 		CreatePatchInfoChildSlot();
 	}
+
+	return Content;
 }
 
 void SDMXPixelMappingOutputComponent::CreateComponentNameChildSlotAbove()
diff --git a/Engine/Plugins/VirtualProduction/DMX/DMXPixelMapping/Source/DMXPixelMappingEditor/Private/Widgets/SDMXPixelMappingOutputComponent.h b/Engine/Plugins/VirtualProduction/DMX/DMXPixelMapping/Source/DMXPixelMappingEditor/Private/Widgets/SDMXPixelMappingOutputComponent.h
--- a/Engine/Plugins/VirtualProduction/DMX/DMXPixelMapping/Source/DMXPixelMappingEditor/Private/Widgets/SDMXPixelMappingOutputComponent.h
+++ b/Engine/Plugins/VirtualProduction/DMX/DMXPixelMapping/Source/DMXPixelMappingEditor/Private/Widgets/SDMXPixelMappingOutputComponent.h
@@ -2,7 +2,7 @@
 
 #pragma once
 
-
+#include "Engine/EngineTypes.h"
 #include "Styling/SlateBrush.h"
 #include "Widgets/SCompoundWidget.h"
 #include "Widgets/Layout/SConstraintCanvas.h"
@@ -72,8 +72,14 @@ protected:
 	//~ End IDMXPixelMappingOutputComponentWidgetInterface
 
 private:
-	/** Updates the child slots */
-	void UpdateChildSlots();
+	/** Refreshes the widget oo the next tick */
+	void RequestRefresh();
+
+	/** Refreshes the widget */
+	void ForceRefresh();
+
+	/** Creates content */
+	TSharedRef<SWidget> CreateContent();
 
 	/** Creates the child slot that displays the component name, above the child slot */
 	void CreateComponentNameChildSlotAbove();
@@ -108,6 +114,9 @@ private:
 	/** Border for content to display names in the top row of the widget */
 	TSharedPtr<SBorder> BottomContentBorder;
 
+	/** Timer handle for request refresh */
+	FTimerHandle RefreshTimerHandle;
+
 	/** The model for this widget */
 	TSharedPtr<FDMXPixelMappingOutputComponentModel> Model;
 
diff --git a/Engine/Plugins/VirtualProduction/DMX/DMXPixelMapping/Source/DMXPixelMappingEditor/Private/Widgets/SDMXPixelMappingPreviewViewport.cpp b/Engine/Plugins/VirtualProduction/DMX/DMXPixelMapping/Source/DMXPixelMappingEditor/Private/Widgets/SDMXPixelMappingPreviewViewport.cpp
--- a/Engine/Plugins/VirtualProduction/DMX/DMXPixelMapping/Source/DMXPixelMappingEditor/Private/Widgets/SDMXPixelMappingPreviewViewport.cpp
+++ b/Engine/Plugins/VirtualProduction/DMX/DMXPixelMapping/Source/DMXPixelMappingEditor/Private/Widgets/SDMXPixelMappingPreviewViewport.cpp
@@ -3,89 +3,84 @@
 #include "Widgets/SDMXPixelMappingPreviewViewport.h"
 
 #include "Components/DMXPixelMappingRendererComponent.h"
+#include "Engine/TextureRenderTarget2D.h"
+#include "Framework/Application/SlateApplication.h"
 #include "Toolkits/DMXPixelMappingToolkit.h"
 #include "Viewports/DMXPixelMappingSceneViewport.h"
 #include "Viewports/DMXPixelMappingPreviewViewportClient.h"
-
-#include "Framework/Application/SlateApplication.h"
-#include "Engine/TextureRenderTarget2D.h"
 #include "Widgets/SViewport.h"
 #include "Widgets/Layout/SBox.h"
 
-void SDMXPixelMappingPreviewViewport::Construct(const FArguments& InArgs, const TSharedPtr<FDMXPixelMappingToolkit>& InToolkit)
+
+void SDMXPixelMappingPreviewViewport::Construct(const FArguments& InArgs, TWeakPtr<FDMXPixelMappingToolkit> InWeakToolkit)
 {
-	bIsRenderingEnabled = true;
-	ToolkitWeakPtr = InToolkit;
+	WeakToolkit = InWeakToolkit;
 
 	ChildSlot
 		[
 			SNew(SBox)
 			.Visibility(EVisibility::HitTestInvisible)
-			.WidthOverride(this, &SDMXPixelMappingPreviewViewport::GetPreviewAreaWidth)
-			.HeightOverride(this, &SDMXPixelMappingPreviewViewport::GetPreviewAreaHeight)
+			.Padding(this, &SDMXPixelMappingPreviewViewport::GetPaddingGraphSpace)
+			.WidthOverride(this, &SDMXPixelMappingPreviewViewport::GetWidthGraphSpace)
+			.HeightOverride(this, &SDMXPixelMappingPreviewViewport::GetHeightGraphSpace)
 			[
 				SAssignNew(ViewportWidget, SViewport)
-					.IsEnabled(FSlateApplication::Get().GetNormalExecutionAttribute())
-					.EnableGammaCorrection(false)
-					.ShowEffectWhenDisabled(false)
+				.IsEnabled(FSlateApplication::Get().GetNormalExecutionAttribute())
+				.EnableGammaCorrection(false)
+				.ShowEffectWhenDisabled(false)
 			]
 		];
 
-	ViewportClient = MakeShared<FDMXPixelMappingPreviewViewportClient>(InToolkit, SharedThis(this));
-
+	ViewportClient = MakeShared<FDMXPixelMappingPreviewViewportClient>(InWeakToolkit, SharedThis(this));
 	Viewport = MakeShared<FDMXPixelMappingSceneViewport>(ViewportClient.Get(), ViewportWidget);
 
-	// The viewport widget needs an interface so it knows what should render
 	ViewportWidget->SetViewportInterface(Viewport.ToSharedRef());
 }
 
-void SDMXPixelMappingPreviewViewport::EnableRendering()
+FOptionalSize SDMXPixelMappingPreviewViewport::GetWidthGraphSpace() const
 {
-	bIsRenderingEnabled = true;
+	if (ViewportClient->IsDrawingVisibleRectOnly())
+	{
+		return ViewportClient->GetVisibleTextureBoxGraphSpace().Max.X;
+	}
+	else if (UTexture* InputTexture = GetInputTexture())
+	{
+		return InputTexture->GetSurfaceWidth();
+	}
+	return FOptionalSize();
 }
 
-void SDMXPixelMappingPreviewViewport::DisableRendering()
+FOptionalSize SDMXPixelMappingPreviewViewport::GetHeightGraphSpace() const
 {
-	bIsRenderingEnabled = false;
+	if (ViewportClient->IsDrawingVisibleRectOnly())
+	{
+		return ViewportClient->GetVisibleTextureBoxGraphSpace().Max.Y;
+	}
+	else if (UTexture* InputTexture = GetInputTexture())
+	{
+		return InputTexture->GetSurfaceHeight();
+	}
+	return FOptionalSize();
 }
 
 void SDMXPixelMappingPreviewViewport::Tick(const FGeometry& AllottedGeometry, const double InCurrentTime, const float InDeltaTime)
 {
-	if (bIsRenderingEnabled)
-	{
-		Viewport->Invalidate();
-	}
+	Viewport->Invalidate();
 }
 
-FOptionalSize SDMXPixelMappingPreviewViewport::GetPreviewAreaWidth() const
+FMargin SDMXPixelMappingPreviewViewport::GetPaddingGraphSpace() const
 {
-	const TSharedPtr<FDMXPixelMappingToolkit> Toolkit = ToolkitWeakPtr.Pin();
-	check(Toolkit.IsValid());
-
-	if (UDMXPixelMappingRendererComponent* RendererComponent = Toolkit->GetActiveRendererComponent())
+	if (ViewportClient->IsDrawingVisibleRectOnly())
 	{
-		UTextureRenderTarget2D* OutputTexture = RendererComponent->GetPreviewRenderTarget();
-		check(OutputTexture) // Preview Texture should be valid in any case
-		
-		return OutputTexture->SizeX;
+		return FMargin(ViewportClient->GetVisibleTextureBoxGraphSpace().Min.X, ViewportClient->GetVisibleTextureBoxGraphSpace().Min.Y, 0.0, 0.0);
 	}
 
-	return 1.f;
+	return FMargin(0.f);
 }
 
-FOptionalSize SDMXPixelMappingPreviewViewport::GetPreviewAreaHeight() const
+UTexture* SDMXPixelMappingPreviewViewport::GetInputTexture() const
 {
-	const TSharedPtr<FDMXPixelMappingToolkit> Toolkit = ToolkitWeakPtr.Pin();
-	check(Toolkit.IsValid());
-
-	if (UDMXPixelMappingRendererComponent* RendererComponent = Toolkit->GetActiveRendererComponent())
-	{
-		UTextureRenderTarget2D* OutputTexture = RendererComponent->GetPreviewRenderTarget();
-		check(OutputTexture) // Preview Texture should be valid in any case
-		
-		return OutputTexture->SizeY;
-	}
-
-	return 1.f;
+	UDMXPixelMappingRendererComponent* RendererComponent = WeakToolkit.IsValid() ? WeakToolkit.Pin()->GetActiveRendererComponent() : nullptr;
+	
+	return RendererComponent ? RendererComponent->GetRenderedInputTexture() : nullptr;
 }
-
diff --git a/Engine/Plugins/VirtualProduction/DMX/DMXPixelMapping/Source/DMXPixelMappingEditor/Private/Widgets/SDMXPixelMappingPreviewViewport.h b/Engine/Plugins/VirtualProduction/DMX/DMXPixelMapping/Source/DMXPixelMappingEditor/Private/Widgets/SDMXPixelMappingPreviewViewport.h
--- a/Engine/Plugins/VirtualProduction/DMX/DMXPixelMapping/Source/DMXPixelMappingEditor/Private/Widgets/SDMXPixelMappingPreviewViewport.h
+++ b/Engine/Plugins/VirtualProduction/DMX/DMXPixelMapping/Source/DMXPixelMappingEditor/Private/Widgets/SDMXPixelMappingPreviewViewport.h
@@ -27,27 +27,29 @@ public:
 	 *
 	 * @param InArgs The construction arguments.
 	 */
-	void Construct(const FArguments& InArgs, const TSharedPtr<FDMXPixelMappingToolkit>& InToolkit);
-
-	/** Enable viewport rendering */
-	void EnableRendering();
-
-	/** Disable viewport rendering */
-	void DisableRendering();
+	void Construct(const FArguments& InArgs, TWeakPtr<FDMXPixelMappingToolkit> InToolkit);
 
 	TSharedPtr<FSceneViewport> GetViewport() const { return Viewport; }
 
 	TSharedPtr<SViewport> GetViewportWidget() const { return ViewportWidget; }
 
+	/** Returns the texture width, in graph space  */
+	FOptionalSize GetWidthGraphSpace() const;
+
+	/** Returns the texture height, in graph space  */
+	FOptionalSize GetHeightGraphSpace() const;
+
+protected:
 	//~ Begin SWidget interface
 	virtual void Tick(const FGeometry& AllottedGeometry, const double InCurrentTime, const float InDeltaTime) override;
 	//~ End of SWidget interface
 
-	FOptionalSize GetPreviewAreaWidth() const;
-
-	FOptionalSize GetPreviewAreaHeight() const;
-
 private:
+	/** Returns the padding of content, in graph space */
+	FMargin GetPaddingGraphSpace() const;
+
+	/** Returns the input texture currently in use */
+	UTexture* GetInputTexture() const;
 
 	/** Viewport client. */
 	TSharedPtr<FDMXPixelMappingPreviewViewportClient> ViewportClient;
@@ -59,8 +61,5 @@ private:
 	TSharedPtr<SViewport> ViewportWidget;
 
 	/** Weak pointer to the editor toolkit */
-	TWeakPtr<FDMXPixelMappingToolkit> ToolkitWeakPtr;
-
-	/** Is rendering currently enabled? (disabled when reimporting a texture) */
-	bool bIsRenderingEnabled;
+	TWeakPtr<FDMXPixelMappingToolkit> WeakToolkit;
 };
diff --git a/Engine/Plugins/VirtualProduction/DMX/DMXPixelMapping/Source/DMXPixelMappingEditor/Private/Widgets/SDMXPixelMappingSourceTextureViewport.cpp b/Engine/Plugins/VirtualProduction/DMX/DMXPixelMapping/Source/DMXPixelMappingEditor/Private/Widgets/SDMXPixelMappingSourceTextureViewport.cpp
--- a/Engine/Plugins/VirtualProduction/DMX/DMXPixelMapping/Source/DMXPixelMappingEditor/Private/Widgets/SDMXPixelMappingSourceTextureViewport.cpp
+++ b/Engine/Plugins/VirtualProduction/DMX/DMXPixelMapping/Source/DMXPixelMappingEditor/Private/Widgets/SDMXPixelMappingSourceTextureViewport.cpp
@@ -40,14 +40,9 @@ void SDMXPixelMappingSourceTextureViewport::Construct(const FArguments& InArgs,
 	ViewportWidget->SetViewportInterface(Viewport.ToSharedRef());
 }
 
-void SDMXPixelMappingSourceTextureViewport::Tick(const FGeometry& AllottedGeometry, const double InCurrentTime, const float InDeltaTime)
-{
-	Viewport->Invalidate();
-}
-
 FOptionalSize SDMXPixelMappingSourceTextureViewport::GetWidthGraphSpace() const
 {
-	if (ViewportClient->DrawOnlyVisibleRect())
+	if (ViewportClient->IsDrawingVisibleRectOnly())
 	{
 		return ViewportClient->GetVisibleTextureBoxGraphSpace().Max.X;
 	}
@@ -60,7 +55,7 @@ FOptionalSize SDMXPixelMappingSourceTextureViewport::GetWidthGraphSpace() const
 
 FOptionalSize SDMXPixelMappingSourceTextureViewport::GetHeightGraphSpace() const
 {
-	if (ViewportClient->DrawOnlyVisibleRect())
+	if (ViewportClient->IsDrawingVisibleRectOnly())
 	{
 		return ViewportClient->GetVisibleTextureBoxGraphSpace().Max.Y;
 	}
@@ -71,9 +66,14 @@ FOptionalSize SDMXPixelMappingSourceTextureViewport::GetHeightGraphSpace() const
 	return FOptionalSize();
 }
 
+void SDMXPixelMappingSourceTextureViewport::Tick(const FGeometry& AllottedGeometry, const double InCurrentTime, const float InDeltaTime)
+{
+	Viewport->Invalidate();
+}
+
 FMargin SDMXPixelMappingSourceTextureViewport::GetPaddingGraphSpace() const
 {
-	if (ViewportClient->DrawOnlyVisibleRect())
+	if (ViewportClient->IsDrawingVisibleRectOnly())
 	{
 		return FMargin(ViewportClient->GetVisibleTextureBoxGraphSpace().Min.X, ViewportClient->GetVisibleTextureBoxGraphSpace().Min.Y, 0.0, 0.0);
 	}
@@ -85,12 +85,7 @@ FMargin SDMXPixelMappingSourceTextureViewport::GetPaddingGraphSpace() const
 
 UTexture* SDMXPixelMappingSourceTextureViewport::GetInputTexture() const
 {
-	if (TSharedPtr<FDMXPixelMappingToolkit> Toolkit = WeakToolkit.Pin())
-	{
-		if (UDMXPixelMappingRendererComponent* RendererComponent = Toolkit->GetActiveRendererComponent())
-		{
-			return RendererComponent->GetRenderedInputTexture();
-		}
-	}
-	return nullptr;
+	UDMXPixelMappingRendererComponent* RendererComponent = WeakToolkit.IsValid() ? WeakToolkit.Pin()->GetActiveRendererComponent() : nullptr;
+
+	return RendererComponent ? RendererComponent->GetRenderedInputTexture() : nullptr;
 }
diff --git a/Engine/Plugins/VirtualProduction/DMX/DMXPixelMapping/Source/DMXPixelMappingEditor/Private/Widgets/SDMXPixelMappingSourceTextureViewport.h b/Engine/Plugins/VirtualProduction/DMX/DMXPixelMapping/Source/DMXPixelMappingEditor/Private/Widgets/SDMXPixelMappingSourceTextureViewport.h
--- a/Engine/Plugins/VirtualProduction/DMX/DMXPixelMapping/Source/DMXPixelMappingEditor/Private/Widgets/SDMXPixelMappingSourceTextureViewport.h
+++ b/Engine/Plugins/VirtualProduction/DMX/DMXPixelMapping/Source/DMXPixelMappingEditor/Private/Widgets/SDMXPixelMappingSourceTextureViewport.h
@@ -22,28 +22,27 @@ public:
 	SLATE_END_ARGS()
 
 public:
-	/**
-	 * Constructs the widget.
-	 *
-	 * @param InArgs The construction arguments.
-	 */
+	/** Constructs this widget */
 	void Construct(const FArguments& InArgs, const TSharedPtr<FDMXPixelMappingToolkit>& InToolkit);
 
+	/** Returns the viewport of this widget */
 	TSharedPtr<FSceneViewport> GetViewport() const { return Viewport; }
 
+	/** Returns the viewport widget of this viewport */
 	TSharedPtr<SViewport> GetViewportWidget() const { return ViewportWidget; }
 
-	// Begin SWidget interface
-	virtual void Tick(const FGeometry& AllottedGeometry, const double InCurrentTime, const float InDeltaTime) override;
-	// End of SWidget interface
-
-private:
 	/** Returns the texture width, in graph space  */
 	FOptionalSize GetWidthGraphSpace() const;
 
 	/** Returns the texture height, in graph space  */
 	FOptionalSize GetHeightGraphSpace() const;
 
+protected:
+	// Begin SWidget interface
+	virtual void Tick(const FGeometry& AllottedGeometry, const double InCurrentTime, const float InDeltaTime) override;
+	// End of SWidget interface
+
+private:
 	/** Returns the padding of content, in graph space */
 	FMargin GetPaddingGraphSpace() const;
 
diff --git a/Engine/Plugins/VirtualProduction/DMX/DMXPixelMapping/Source/DMXPixelMappingEditor/Private/Widgets/SDMXPixelMappingSurface.cpp b/Engine/Plugins/VirtualProduction/DMX/DMXPixelMapping/Source/DMXPixelMappingEditor/Private/Widgets/SDMXPixelMappingSurface.cpp
--- a/Engine/Plugins/VirtualProduction/DMX/DMXPixelMapping/Source/DMXPixelMappingEditor/Private/Widgets/SDMXPixelMappingSurface.cpp
+++ b/Engine/Plugins/VirtualProduction/DMX/DMXPixelMapping/Source/DMXPixelMappingEditor/Private/Widgets/SDMXPixelMappingSurface.cpp
@@ -2,10 +2,10 @@
 
 #include "Widgets/SDMXPixelMappingSurface.h"
 
-#include "Toolkits/DMXPixelMappingToolkit.h"
-
 #include "Framework/Application/SlateApplication.h"
+#include "Components/DMXPixelMappingRendererComponent.h"
 #include "Settings/LevelEditorViewportSettings.h"
+#include "Toolkits/DMXPixelMappingToolkit.h"
 
 
 #define LOCTEXT_NAMESPACE "SDMXPixelMappingSurface"
@@ -131,7 +131,7 @@ struct FFixedZoomLevelsContainerDesignSurface : public FZoomLevelsContainer
 
 void SDMXPixelMappingSurface::Construct(const FArguments& InArgs, const TSharedPtr<FDMXPixelMappingToolkit>& InToolkit)
 {
-	ToolkitWeakPtr = InToolkit;
+	WeakToolkit = InToolkit;
 
 	if (!ZoomLevels)
 	{
@@ -143,7 +143,6 @@ void SDMXPixelMappingSurface::Construct(const FArguments& InArgs, const TSharedP
 	AllowContinousZoomInterpolation = InArgs._AllowContinousZoomInterpolation;
 	bIsPanning = false;
 	bIsZooming = false;
-	bRecenteredOnFirstTick = false;
 
 	ViewOffset = FVector2D::ZeroVector;
 	bDrawGridLines = true;
@@ -174,8 +173,6 @@ void SDMXPixelMappingSurface::Construct(const FArguments& InArgs, const TSharedP
 		[
 			InArgs._Content.Widget
 		];
-
-	InitialBounds = ComputeAreaBounds();
 }
 
 EActiveTimerReturnType SDMXPixelMappingSurface::HandleZoomToFit(double InCurrentTime, float InDeltaTime)
@@ -200,6 +197,19 @@ EActiveTimerReturnType SDMXPixelMappingSurface::HandleZoomToFit(double InCurrent
 
 void SDMXPixelMappingSurface::Tick(const FGeometry& AllottedGeometry, const double InCurrentTime, const float InDeltaTime)
 {
+	if (bInitialZoomToFit && WeakToolkit.IsValid())
+	{
+		const TSharedRef<FDMXPixelMappingToolkit> Toolkit = WeakToolkit.Pin().ToSharedRef();
+		UDMXPixelMappingRendererComponent* RendererComponent = Toolkit->GetActiveRendererComponent();
+		UTexture* Texture = RendererComponent ? RendererComponent->GetRenderedInputTexture() : nullptr;
+		if (Texture)
+		{
+			constexpr bool bInstantZoom = true;
+			ZoomToFit(bInstantZoom);
+			bInitialZoomToFit = false;
+		}
+	}
+
 	if (bDeferredZoomToExtents)
 	{
 		const FSlateRect Bounds = ComputeAreaBounds();
@@ -212,12 +222,6 @@ void SDMXPixelMappingSurface::Tick(const FGeometry& AllottedGeometry, const doub
 			RegisterActiveTimer(0.f, FWidgetActiveTimerDelegate::CreateSP(this, &SDMXPixelMappingSurface::HandleZoomToFit));
 		}
 	}
-
-	if (!bRecenteredOnFirstTick && InitialBounds != ComputeAreaBounds())
-	{
-		ZoomToFit(true);
-		bRecenteredOnFirstTick = true;
-	}
 }
 
 FCursorReply SDMXPixelMappingSurface::OnCursorQuery(const FGeometry& MyGeometry, const FPointerEvent& CursorEvent) const
diff --git a/Engine/Plugins/VirtualProduction/DMX/DMXPixelMapping/Source/DMXPixelMappingEditor/Private/Widgets/SDMXPixelMappingSurface.h b/Engine/Plugins/VirtualProduction/DMX/DMXPixelMapping/Source/DMXPixelMappingEditor/Private/Widgets/SDMXPixelMappingSurface.h
--- a/Engine/Plugins/VirtualProduction/DMX/DMXPixelMapping/Source/DMXPixelMappingEditor/Private/Widgets/SDMXPixelMappingSurface.h
+++ b/Engine/Plugins/VirtualProduction/DMX/DMXPixelMapping/Source/DMXPixelMappingEditor/Private/Widgets/SDMXPixelMappingSurface.h
@@ -52,7 +52,7 @@ public:
 	float GetZoomAmount() const;
 
 	/** Returns the toolkit used with this surface */
-	FORCEINLINE TSharedPtr<FDMXPixelMappingToolkit> GetToolkit() { return ToolkitWeakPtr.Pin(); }
+	FORCEINLINE TSharedPtr<FDMXPixelMappingToolkit> GetToolkit() { return WeakToolkit.Pin(); }
 
 	/** Zooms to fit content */
 	void ZoomToFit(bool bInstantZoom);
@@ -144,12 +144,12 @@ protected:
 	/** Does the user need to press Control in order to over-zoom. */
 	bool bRequireControlToOverZoom;
 
-	/** Initial bounds, useful to detect when the initial zoom to fit should occur */
-	FSlateRect InitialBounds;
+	/** If true, zooms to fit as soon as a texture is loaded */
+	bool bInitialZoomToFit = true;
 
 	/** Toolkit weak pointer */
-	TWeakPtr<FDMXPixelMappingToolkit> ToolkitWeakPtr;
-
+	TWeakPtr<FDMXPixelMappingToolkit> WeakToolkit;
+	
 private:
 	/** Active timer that handles deferred zooming until the target zoom is reached */
 	EActiveTimerReturnType HandleZoomToFit(double InCurrentTime, float InDeltaTime);
@@ -159,7 +159,4 @@ private:
 
 	// A flag noting if we have a pending zoom to extents operation to perform next tick.
 	bool bDeferredZoomToExtents;
-
-	/** Recenter the graph an first tick */
-	bool bRecenteredOnFirstTick;
 };
diff --git a/Engine/Plugins/VirtualProduction/DMX/DMXPixelMapping/Source/DMXPixelMappingEditor/Private/Widgets/SDMXPixelMappingTransformHandle.cpp b/Engine/Plugins/VirtualProduction/DMX/DMXPixelMapping/Source/DMXPixelMappingEditor/Private/Widgets/SDMXPixelMappingTransformHandle.cpp
--- a/Engine/Plugins/VirtualProduction/DMX/DMXPixelMapping/Source/DMXPixelMappingEditor/Private/Widgets/SDMXPixelMappingTransformHandle.cpp
+++ b/Engine/Plugins/VirtualProduction/DMX/DMXPixelMapping/Source/DMXPixelMappingEditor/Private/Widgets/SDMXPixelMappingTransformHandle.cpp
@@ -2,13 +2,13 @@
 
 #include "Widgets/SDMXPixelMappingTransformHandle.h"
 
-#include "DMXPixelMappingLayoutSettings.h"
 #include "Components/DMXPixelMappingMatrixCellComponent.h"
+#include "DMXPixelMappingLayoutSettings.h"
+#include "Framework/Application/SlateApplication.h"
+#include "ScopedTransaction.h"
 #include "Toolkits/DMXPixelMappingToolkit.h"
 #include "Views/SDMXPixelMappingDesignerView.h"
-
 #include "Widgets/Images/SImage.h"
-#include "ScopedTransaction.h"
 
 
 #define LOCTEXT_NAMESPACE "SDMXPixelMappingTransformHandle"
@@ -40,8 +40,19 @@ EVisibility SDMXPixelMappingTransformHandle::GetHandleVisibility() const
 
 FReply SDMXPixelMappingTransformHandle::OnMouseButtonDown(const FGeometry& MyGeometry, const FPointerEvent& MouseEvent)
 {
-	if ( MouseEvent.GetEffectingButton() == EKeys::LeftMouseButton )
+	if (MouseEvent.GetEffectingButton() == EKeys::LeftMouseButton)
 	{
+		const FVector2D LocalSize = MyGeometry.GetLocalSize();
+		const FVector2D LocalCursorPos = MyGeometry.AbsoluteToLocal(MouseEvent.GetScreenSpacePosition());
+
+		// Only handle dragging the edges of the widget
+		constexpr double MouseThreshold = 10.f;
+		if (LocalCursorPos.X < LocalSize.X - MouseThreshold ||
+			LocalCursorPos.Y < LocalSize.Y - MouseThreshold)
+		{
+			return FReply::Unhandled();;
+		}
+
 		Action = ComputeActionAtLocation(MyGeometry, MouseEvent);
 
 		if (TSharedPtr<FDMXPixelMappingToolkit> Toolkit = DesignerViewWeakPtr.Pin()->GetToolkit())
@@ -49,10 +60,9 @@ FReply SDMXPixelMappingTransformHandle::OnMouseButtonDown(const FGeometry& MyGeo
 			const TSet<FDMXPixelMappingComponentReference> SelectedComponentReferences = Toolkit->GetSelectedComponents();
 			for (const FDMXPixelMappingComponentReference& ComponentReference : SelectedComponentReferences)
 			{
-				UDMXPixelMappingBaseComponent* Preview = ComponentReference.GetComponent();
 				UDMXPixelMappingBaseComponent* Component = ComponentReference.GetComponent();
 
-				if (UDMXPixelMappingOutputComponent* OutputComponent = Cast<UDMXPixelMappingOutputComponent>(Preview))
+				if (UDMXPixelMappingOutputComponent* OutputComponent = Cast<UDMXPixelMappingOutputComponent>(Component))
 				{
 					FMargin Offsets;
 					FVector2D Size = OutputComponent->GetSize();
@@ -107,7 +117,7 @@ FReply SDMXPixelMappingTransformHandle::OnMouseButtonUp(const FGeometry& MyGeome
 
 FReply SDMXPixelMappingTransformHandle::OnMouseMove(const FGeometry& MyGeometry, const FPointerEvent& MouseEvent)
 {
-	if ( Action != EDMXPixelMappingTransformAction::None )
+	if (Action != EDMXPixelMappingTransformAction::None && DesignerViewWeakPtr.IsValid())
 	{
 		if (TSharedPtr<FDMXPixelMappingToolkit> Toolkit = DesignerViewWeakPtr.Pin()->GetToolkit())
 		{
@@ -116,10 +126,7 @@ FReply SDMXPixelMappingTransformHandle::OnMouseMove(const FGeometry& MyGeometry,
 			{
 				if (UDMXPixelMappingOutputComponent* ResizedComponent = Cast<UDMXPixelMappingOutputComponent>(ComponentReference.GetComponent()))
 				{
-					const FVector2D Delta = MouseEvent.GetScreenSpacePosition() - MouseDownPosition;
-					const FVector2D TranslateAmount = Delta * (1.0f / (DesignerViewWeakPtr.Pin()->GetZoomAmount() * MyGeometry.Scale));
-
-					Resize(ResizedComponent, DragDirection, TranslateAmount);
+					RequestResize(ResizedComponent, DragDirection);
 				}
 			}
 		}
@@ -128,28 +135,40 @@ FReply SDMXPixelMappingTransformHandle::OnMouseMove(const FGeometry& MyGeometry,
 	return FReply::Unhandled();
 }
 
+void SDMXPixelMappingTransformHandle::RequestResize(UDMXPixelMappingBaseComponent* BaseComponent, const FVector2D& Direction)
+{
+	if (!RequestResizeHandle.IsValid() && DesignerViewWeakPtr.IsValid())
+	{
+		const FVector2D& CursorPosition = FSlateApplication::Get().GetCursorPos();
+		const FVector2D Delta = CursorPosition - MouseDownPosition;
+		const FVector2D TranslateAmount = Delta * (1.0f / (DesignerViewWeakPtr.Pin()->GetZoomAmount() * GetCachedGeometry().Scale));
+
+		RequestResizeHandle = GEditor->GetTimerManager()->SetTimerForNextTick(FTimerDelegate::CreateLambda([this, BaseComponent, Direction, TranslateAmount]()
+			{
+				Resize(BaseComponent, Direction, TranslateAmount);
+				RequestResizeHandle.Invalidate();
+			}));
+	}
+
+}
+
 void SDMXPixelMappingTransformHandle::Resize(UDMXPixelMappingBaseComponent* BaseComponent, const FVector2D& Direction, const FVector2D& Amount)
 {
 	if (UDMXPixelMappingOutputComponent* OutputComponent = Cast<UDMXPixelMappingOutputComponent>(BaseComponent))
 	{
-		FVector2D ComponentSize = OutputComponent->GetSize();
-
 		FMargin Offsets = StartingOffsets;
 
 		const FVector2D Movement = Amount * Direction;
-		const FVector2D PositionMovement = Movement * (FVector2D(1.0f, 1.0f));
-		const FVector2D SizeMovement = Movement;
-
 		if (Direction.X < 0)
 		{
-			Offsets.Left -= PositionMovement.X;
-			Offsets.Right += SizeMovement.X;
+			Offsets.Left -= Movement.X;
+			Offsets.Right += Movement.X;
 		}
 
 		if (Direction.Y < 0)
 		{
-			Offsets.Top -= PositionMovement.Y;
-			Offsets.Bottom += SizeMovement.Y;
+			Offsets.Top -= Movement.Y;
+			Offsets.Bottom += Movement.Y;
 		}
 
 		if (Direction.X > 0)
diff --git a/Engine/Plugins/VirtualProduction/DMX/DMXPixelMapping/Source/DMXPixelMappingEditor/Private/Widgets/SDMXPixelMappingTransformHandle.h b/Engine/Plugins/VirtualProduction/DMX/DMXPixelMapping/Source/DMXPixelMappingEditor/Private/Widgets/SDMXPixelMappingTransformHandle.h
--- a/Engine/Plugins/VirtualProduction/DMX/DMXPixelMapping/Source/DMXPixelMappingEditor/Private/Widgets/SDMXPixelMappingTransformHandle.h
+++ b/Engine/Plugins/VirtualProduction/DMX/DMXPixelMapping/Source/DMXPixelMappingEditor/Private/Widgets/SDMXPixelMappingTransformHandle.h
@@ -2,8 +2,10 @@
 
 #pragma once
 
+#include "Engine/EngineTypes.h"
 #include "Widgets/SCompoundWidget.h"
 
+
 enum class EDMXPixelMappingTransformDirection : uint8
 {
 	CenterRight,
@@ -55,7 +57,10 @@ public:
 protected:
 	EVisibility GetHandleVisibility() const;
 
-	bool CanResize(UDMXPixelMappingBaseComponent* BaseComponent, const FVector2D& Direction) const { return true; }
+	/** Resizes the component on the next tick */
+	void RequestResize(UDMXPixelMappingBaseComponent* BaseComponent, const FVector2D& Direction);
+	
+	/** Resizes the component */
 	void Resize(UDMXPixelMappingBaseComponent* BaseComponent, const FVector2D& Direction, const FVector2D& Amount);
 
 	EDMXPixelMappingTransformAction ComputeActionAtLocation(const FGeometry& MyGeometry, const FPointerEvent& MouseEvent) const;
@@ -79,4 +84,7 @@ protected:
 	TWeakPtr<SDMXPixelMappingDesignerView> DesignerViewWeakPtr;
 
 	TAttribute<FVector2D> Offset;
+
+	/** Timer handle used when resize is requested */
+	FTimerHandle RequestResizeHandle;
 };
diff --git a/Engine/Plugins/VirtualProduction/DMX/DMXPixelMapping/Source/DMXPixelMappingRenderer/Private/DMXPixelMappingPixelMapRenderer.cpp b/Engine/Plugins/VirtualProduction/DMX/DMXPixelMapping/Source/DMXPixelMappingRenderer/Private/DMXPixelMappingPixelMapRenderer.cpp
new file mode 100644
--- /dev/null
+++ b/Engine/Plugins/VirtualProduction/DMX/DMXPixelMapping/Source/DMXPixelMappingRenderer/Private/DMXPixelMappingPixelMapRenderer.cpp
@@ -0,0 +1,72 @@
+// Copyright Epic Games, Inc. All Rights Reserved.
+
+#include "DMXPixelMappingPixelMapRenderer.h"
+
+#include "Engine/TextureRenderTarget2D.h"
+#include "PixelMapRenderElement.h"
+#include "RenderPixelMapProxy.h"
+#include "RenderingThread.h"
+
+
+const FIntPoint UDMXPixelMappingPixelMapRenderer::MaxPixelMapSize = FIntPoint(4096, 4096);
+
+UDMXPixelMappingPixelMapRenderer::UDMXPixelMappingPixelMapRenderer()
+{
+	using namespace UE::DMXPixelMapping::Rendering::Private;
+	RendererPixelMapProxy = MakeShared<FRenderPixelMapProxy>();
+}
+
+void UDMXPixelMappingPixelMapRenderer::SetElements(const TArray<TSharedRef<UE::DMXPixelMapping::Rendering::FPixelMapRenderElement>>& InElements)
+{
+	check(IsInGameThread());
+
+	if (RenderElements == InElements)
+	{
+		return;
+	}
+
+	if (InElements.IsEmpty())
+	{
+		PixelMapRenderTarget = nullptr;
+		return;
+	}
+
+	if (!ensureMsgf(InElements.Num() < MaxPixelMapSize.X * MaxPixelMapSize.Y, TEXT("PixelMapping contains more pixels than supported by engine.")))
+	{
+		return;
+	}
+
+	const FIntPoint DesiredSize
+	{
+		FMath::Min(InElements.Num(), MaxPixelMapSize.X),
+		InElements.Num() / MaxPixelMapSize.X + 1
+	};
+	const FVector2D ValidSize
+	{
+		FMath::Max(1.0, DesiredSize.X),
+		FMath::Max(1.0, DesiredSize.Y)
+	};
+
+	// Create a new render target if needed
+	if (!PixelMapRenderTarget)
+	{
+		const FName TargetName = MakeUniqueObjectName(this, UTextureRenderTarget2D::StaticClass(), TEXT("PixelMapRenderTarget"));
+		PixelMapRenderTarget = NewObject<UTextureRenderTarget2D>(this, TargetName);
+		PixelMapRenderTarget->ClearColor = FLinearColor::Black;
+
+		constexpr bool bInForceLinearGamma = false;
+		const FVector2D InitialSize{ 1.0, 1.0 };
+		PixelMapRenderTarget->InitCustomFormat(ValidSize.X, ValidSize.Y, EPixelFormat::PF_B8G8R8A8, bInForceLinearGamma);
+	}
+	else if (ValidSize.X != PixelMapRenderTarget->GetSurfaceWidth() || ValidSize.Y != PixelMapRenderTarget->GetSurfaceHeight())
+	{
+		PixelMapRenderTarget->ResizeTarget(ValidSize.X, ValidSize.Y);
+	}
+
+	RenderElements = InElements;
+}
+
+void UDMXPixelMappingPixelMapRenderer::Render(UTexture* InputTexture, float Brightness)
+{
+	RendererPixelMapProxy->Render(InputTexture, PixelMapRenderTarget, RenderElements, Brightness);
+}
diff --git a/Engine/Plugins/VirtualProduction/DMX/DMXPixelMapping/Source/DMXPixelMappingRenderer/Private/DMXPixelMappingPreprocessRenderer.cpp b/Engine/Plugins/VirtualProduction/DMX/DMXPixelMapping/Source/DMXPixelMappingRenderer/Private/DMXPixelMappingPreprocessRenderer.cpp
--- a/Engine/Plugins/VirtualProduction/DMX/DMXPixelMapping/Source/DMXPixelMappingRenderer/Private/DMXPixelMappingPreprocessRenderer.cpp
+++ b/Engine/Plugins/VirtualProduction/DMX/DMXPixelMapping/Source/DMXPixelMappingRenderer/Private/DMXPixelMappingPreprocessRenderer.cpp
@@ -2,8 +2,8 @@
 
 #include "DMXPixelMappingPreprocessRenderer.h"
 
+#include "DMXStats.h"
 #include "Engine/TextureRenderTarget2D.h"
-#include "IDMXPixelMappingRenderer.h"
 #include "Materials/MaterialInstanceDynamic.h"
 #include "PreprocessApplyFilterMaterialProxy.h"
 #include "PreprocessRenderInputMaterialProxy.h"
@@ -12,6 +12,8 @@
 #include "UObject/Package.h"
 
 
+DECLARE_CYCLE_STAT(TEXT("PixelMapping PreprocessInputTexture"), STAT_PreprocessInputTexture, STATGROUP_DMX);
+
 void UDMXPixelMappingPreprocessRenderer::SetInputTexture(UTexture* InTexture)
 {
 	using namespace UE::DMXPixelMapping::Rendering::Preprocess::Private;
@@ -47,6 +49,8 @@ void UDMXPixelMappingPreprocessRenderer::ClearInput()
 
 void UDMXPixelMappingPreprocessRenderer::Render()
 {
+	SCOPE_CYCLE_COUNTER(STAT_PreprocessInputTexture);
+
 	if (RenderInputProxy.IsValid())
 	{
 		RenderInputProxy->Render();
diff --git a/Engine/Plugins/VirtualProduction/DMX/DMXPixelMapping/Source/DMXPixelMappingRenderer/Private/DMXPixelMappingRenderer.cpp b/Engine/Plugins/VirtualProduction/DMX/DMXPixelMapping/Source/DMXPixelMappingRenderer/Private/DMXPixelMappingRenderer.cpp
--- a/Engine/Plugins/VirtualProduction/DMX/DMXPixelMapping/Source/DMXPixelMappingRenderer/Private/DMXPixelMappingRenderer.cpp
+++ b/Engine/Plugins/VirtualProduction/DMX/DMXPixelMapping/Source/DMXPixelMappingRenderer/Private/DMXPixelMappingRenderer.cpp
@@ -27,8 +27,6 @@ namespace DMXPixelMappingRenderer
 	static constexpr auto RenderPassHint = TEXT("Render Pixel Mapping");
 };
 
-DECLARE_GPU_STAT_NAMED(DMXPixelMappingShadersStat, DMXPixelMappingRenderer::RenderPassHint);
-
 #if WITH_EDITOR
 namespace DMXPixelMappingRenderer
 {
@@ -39,40 +37,53 @@ namespace DMXPixelMappingRenderer
 DECLARE_GPU_STAT_NAMED(DMXPixelMappingPreviewStat, DMXPixelMappingRenderer::RenderPreviewPassHint);
 #endif // WITH_EDITOR
 
-class FDMXPixelBlendingQualityDimension : SHADER_PERMUTATION_ENUM_CLASS("PIXELBLENDING_QUALITY", EDMXPixelShaderBlendingQuality);
-class FDMXVertexUVDimension : SHADER_PERMUTATION_BOOL("VERTEX_UV_STATIC_CALCULATION");
+class UE_DEPRECATED(5.3, "Only here to keep support of deprecated FDMXPixelMappingRenderer::DownsampleRender.") FDEPRECATED_DMXPixelBlendingQualityDimension;
+PRAGMA_DISABLE_DEPRECATION_WARNINGS
+class FDEPRECATED_DMXPixelBlendingQualityDimension : SHADER_PERMUTATION_ENUM_CLASS("PIXELBLENDING_QUALITY", EDMXPixelShaderBlendingQuality);
+PRAGMA_ENABLE_DEPRECATION_WARNINGS
+
+class UE_DEPRECATED(5.3, "Only here to keep support of deprecated FDMXPixelMappingRenderer::DownsampleRender.") FDEPRECATED_DMXVertexUVDimension;
+class FDEPRECATED_DMXVertexUVDimension : SHADER_PERMUTATION_BOOL("VERTEX_UV_STATIC_CALCULATION");
 
 /**
  * Pixel Mapping downsampling vertex shader
  */
-class FDMXPixelMappingRendererVS
+class UE_DEPRECATED(5.3, "Only here to keep support of deprecated FDMXPixelMappingRenderer::DownsampleRender.") FDEPRECATED_DMXPixelMappingRendererVS;
+class FDEPRECATED_DMXPixelMappingRendererVS
 	: public FGlobalShader
 {
+	PRAGMA_DISABLE_DEPRECATION_WARNINGS
+
 public:
-	DECLARE_GLOBAL_SHADER(FDMXPixelMappingRendererVS);
-	SHADER_USE_PARAMETER_STRUCT(FDMXPixelMappingRendererVS, FGlobalShader);
+	DECLARE_GLOBAL_SHADER(FDEPRECATED_DMXPixelMappingRendererVS);
+	SHADER_USE_PARAMETER_STRUCT(FDEPRECATED_DMXPixelMappingRendererVS, FGlobalShader);
 
-	using FPermutationDomain = TShaderPermutationDomain<FDMXPixelBlendingQualityDimension, FDMXVertexUVDimension>;
+	using FPermutationDomain = TShaderPermutationDomain<FDEPRECATED_DMXPixelBlendingQualityDimension, FDEPRECATED_DMXVertexUVDimension>;
 	BEGIN_SHADER_PARAMETER_STRUCT(FParameters, )
 		SHADER_PARAMETER(FVector4f, DrawRectanglePosScaleBias)
 		SHADER_PARAMETER(FVector4f, DrawRectangleInvTargetSizeAndTextureSize)
 		SHADER_PARAMETER(FVector4f, DrawRectangleUVScaleBias)
 	END_SHADER_PARAMETER_STRUCT()
 
+	PRAGMA_ENABLE_DEPRECATION_WARNINGS
+
 	static bool ShouldCompilePermutation(const FGlobalShaderPermutationParameters& Parameters) { return true; }
 };
 
 /**
  * Pixel Mapping downsampling pixel shader
  */
-class FDMXPixelMappingRendererPS
+class UE_DEPRECATED(5.3, "Only here to keep support of deprecated FDMXPixelMappingRenderer::DownsampleRender.") FDEPRECATED_DMXPixelMappingRendererPS;
+class FDEPRECATED_DMXPixelMappingRendererPS
 	: public FGlobalShader
 {
+	PRAGMA_DISABLE_DEPRECATION_WARNINGS
+
 public:
-	DECLARE_GLOBAL_SHADER(FDMXPixelMappingRendererPS);
-	SHADER_USE_PARAMETER_STRUCT(FDMXPixelMappingRendererPS, FGlobalShader);
+	DECLARE_GLOBAL_SHADER(FDEPRECATED_DMXPixelMappingRendererPS);
+	SHADER_USE_PARAMETER_STRUCT(FDEPRECATED_DMXPixelMappingRendererPS, FGlobalShader);
 
-	using FPermutationDomain = TShaderPermutationDomain<FDMXPixelBlendingQualityDimension, FDMXVertexUVDimension>;
+	using FPermutationDomain = TShaderPermutationDomain<FDEPRECATED_DMXPixelBlendingQualityDimension, FDEPRECATED_DMXVertexUVDimension>;
 	BEGIN_SHADER_PARAMETER_STRUCT(FParameters, )
 		SHADER_PARAMETER_TEXTURE(Texture2D, InputTexture)
 		SHADER_PARAMETER_SAMPLER(SamplerState, InputSampler)
@@ -83,20 +94,26 @@ public:
 		SHADER_PARAMETER(FVector2f, UVCellSize)
 	END_SHADER_PARAMETER_STRUCT()
 
+	PRAGMA_ENABLE_DEPRECATION_WARNINGS
+
 	static bool ShouldCompilePermutation(const FGlobalShaderPermutationParameters& Parameters) { return true; }
 };
 
-IMPLEMENT_GLOBAL_SHADER(FDMXPixelMappingRendererVS, "/Plugin/DMXPixelMapping/Private/DMXPixelMapping.usf", "DMXPixelMappingVS", SF_Vertex);
-IMPLEMENT_GLOBAL_SHADER(FDMXPixelMappingRendererPS, "/Plugin/DMXPixelMapping/Private/DMXPixelMapping.usf", "DMXPixelMappingPS", SF_Pixel);
+PRAGMA_DISABLE_DEPRECATION_WARNINGS
+IMPLEMENT_GLOBAL_SHADER(FDEPRECATED_DMXPixelMappingRendererVS, "/Plugin/DMXPixelMapping/Private/DMXPixelMapping.usf", "DMXPixelMappingVS", SF_Vertex);
+IMPLEMENT_GLOBAL_SHADER(FDEPRECATED_DMXPixelMappingRendererPS, "/Plugin/DMXPixelMapping/Private/DMXPixelMapping.usf", "DMXPixelMappingPS", SF_Pixel);
+PRAGMA_ENABLE_DEPRECATION_WARNINGS
 
 
+PRAGMA_DISABLE_DEPRECATION_WARNINGS
 FDMXPixelMappingRenderer::FDMXPixelMappingRenderer()
 {
 	static const FName RendererModuleName("Renderer");
 	RendererModule = FModuleManager::GetModulePtr<IRendererModule>(RendererModuleName);
-
 }
+PRAGMA_ENABLE_DEPRECATION_WARNINGS
 
+PRAGMA_DISABLE_DEPRECATION_WARNINGS
 void FDMXPixelMappingRenderer::DownsampleRender(
 	const FTextureResource* InputTexture,
 	const FTextureResource* DstTexture,
@@ -105,15 +122,14 @@ void FDMXPixelMappingRenderer::DownsampleRender(
 	DownsampleReadCallback InCallback
 ) const
 {
+	// DEPRECATED 5.3
+
 	check(IsInGameThread());
 	
 	ENQUEUE_RENDER_COMMAND(DownsampleRenderRDG)(
 		[this, InputTexture, DstTexture, DstTextureTargetResource, InCallback, DownsamplePixelPass = InDownsamplePixelPass]
 		(FRHICommandListImmediate& RHICmdList)
 		{
-			SCOPED_GPU_STAT(RHICmdList, DMXPixelMappingShadersStat);
-			SCOPED_DRAW_EVENTF(RHICmdList, DMXPixelMappingShadersStat, DMXPixelMappingRenderer::RenderPassName);
-			
 			FRHITexture* RenderTargetRef = DstTextureTargetResource->TextureRHI;
 			FRHITexture* DstTextureRef = DstTexture->TextureRHI;
 			FRHITexture* ResolveRenderTarget = DstTextureTargetResource->GetRenderTargetTexture();
@@ -142,14 +158,14 @@ void FDMXPixelMappingRenderer::DownsampleRender(
 				for (const FDMXPixelMappingDownsamplePixelParamsV2& PixelParam : DownsamplePixelPass)
 				{
 					// Create shader permutations
-					FDMXPixelMappingRendererPS::FPermutationDomain PermutationVector;
-					PermutationVector.Set<FDMXPixelBlendingQualityDimension>(static_cast<EDMXPixelShaderBlendingQuality>(PixelParam.CellBlendingQuality));
-					PermutationVector.Set<FDMXVertexUVDimension>(PixelParam.bStaticCalculateUV);
+					FDEPRECATED_DMXPixelMappingRendererPS::FPermutationDomain PermutationVector;
+					PermutationVector.Set<FDEPRECATED_DMXPixelBlendingQualityDimension>(static_cast<EDMXPixelShaderBlendingQuality>(PixelParam.CellBlendingQuality));
+					PermutationVector.Set<FDEPRECATED_DMXVertexUVDimension>(PixelParam.bStaticCalculateUV);
 
 					// Create shaders
 					FGlobalShaderMap* ShaderMap = GetGlobalShaderMap(GMaxRHIFeatureLevel);
-					TShaderMapRef<FDMXPixelMappingRendererVS> VertexShader(ShaderMap, PermutationVector);
-					TShaderMapRef<FDMXPixelMappingRendererPS> PixelShader(ShaderMap, PermutationVector);
+					TShaderMapRef<FDEPRECATED_DMXPixelMappingRendererVS> VertexShader(ShaderMap, PermutationVector);
+					TShaderMapRef<FDEPRECATED_DMXPixelMappingRendererPS> PixelShader(ShaderMap, PermutationVector);
 
 					// Setup graphics pipeline
 					FGraphicsPipelineStateInitializer GraphicsPSOInit;
@@ -164,7 +180,7 @@ void FDMXPixelMappingRenderer::DownsampleRender(
 					SetGraphicsPipelineState(RHICmdList, GraphicsPSOInit, 0);
 
 					// Set vertex shader buffer
-					FDMXPixelMappingRendererVS::FParameters VSParameters;
+					FDEPRECATED_DMXPixelMappingRendererVS::FParameters VSParameters;
 					VSParameters.DrawRectanglePosScaleBias = FVector4f(PixelSize.X, PixelSize.Y, PixelParam.Position.X, PixelParam.Position.Y);
 					VSParameters.DrawRectangleUVScaleBias = FVector4f(PixelParam.UVSize.X, PixelParam.UVSize.Y, PixelParam.UV.X, PixelParam.UV.Y);
 					VSParameters.DrawRectangleInvTargetSizeAndTextureSize = FVector4f(
@@ -173,7 +189,7 @@ void FDMXPixelMappingRenderer::DownsampleRender(
 					SetShaderParameters(RHICmdList, VertexShader, VertexShader.GetVertexShader(), VSParameters);
 					
 					// Set pixel shader buffer
-					FDMXPixelMappingRendererPS::FParameters PSParameters;
+					FDEPRECATED_DMXPixelMappingRendererPS::FParameters PSParameters;
 					PSParameters.InputTexture = InputTextureRHI;
 					PSParameters.Brightness = Brightness;
 					PSParameters.InputSampler = TStaticSamplerState<SF_Trilinear, AM_Clamp, AM_Clamp, AM_Clamp>::GetRHI();
@@ -205,7 +221,7 @@ void FDMXPixelMappingRenderer::DownsampleRender(
 			}
 		});
 }
-
+PRAGMA_ENABLE_DEPRECATION_WARNINGS
 
 void FDMXPixelMappingRenderer::RenderMaterial(UTextureRenderTarget2D* InRenderTarget, UMaterialInterface* InMaterialInterface) const
 {	
diff --git a/Engine/Plugins/VirtualProduction/DMX/DMXPixelMapping/Source/DMXPixelMappingRenderer/Private/DMXPixelMappingRenderer.h b/Engine/Plugins/VirtualProduction/DMX/DMXPixelMapping/Source/DMXPixelMappingRenderer/Private/DMXPixelMappingRenderer.h
--- a/Engine/Plugins/VirtualProduction/DMX/DMXPixelMapping/Source/DMXPixelMappingRenderer/Private/DMXPixelMappingRenderer.h
+++ b/Engine/Plugins/VirtualProduction/DMX/DMXPixelMapping/Source/DMXPixelMappingRenderer/Private/DMXPixelMappingRenderer.h
@@ -12,10 +12,8 @@ class FWidgetRenderer;
 struct FSlateMaterialBrush;
 
 
-
-/**
- * Implementation of Pixel Mapping Renderer
- */
+/** DEPRECATED 5.2 */
+PRAGMA_DISABLE_DEPRECATION_WARNINGS
 class FDMXPixelMappingRenderer
 	: public IDMXPixelMappingRenderer
 {
@@ -32,13 +30,8 @@ public:
 		DownsampleReadCallback InCallback
 	) const override;
 
-PRAGMA_DISABLE_DEPRECATION_WARNINGS
-	UE_DEPRECATED(5.3, "Use PreprocessRenderInputMaterialProxy instead")
 	virtual void RenderMaterial(UTextureRenderTarget2D* InRenderTarget, UMaterialInterface* InMaterialInterface) const override;
-
-	UE_DEPRECATED(5.3, "Use PreprocessRenderInputUserWidgetProxy instead")
 	virtual void RenderWidget(UTextureRenderTarget2D* InRenderTarget, UUserWidget* InUserWidget) const override;
-PRAGMA_ENABLE_DEPRECATION_WARNINGS
 
 	virtual void RenderTextureToRectangle(const FTextureResource* InTextureResource, const FTexture2DRHIRef InRenderTargetTexture, FVector2D InSize, bool bSRGBSource) const override;
 
@@ -60,3 +53,4 @@ private:
 	/** The public interface of the renderer module. */
 	IRendererModule* RendererModule;
 };
+PRAGMA_ENABLE_DEPRECATION_WARNINGS
diff --git a/Engine/Plugins/VirtualProduction/DMX/DMXPixelMapping/Source/DMXPixelMappingRenderer/Private/PixelMapRenderElement.cpp b/Engine/Plugins/VirtualProduction/DMX/DMXPixelMapping/Source/DMXPixelMappingRenderer/Private/PixelMapRenderElement.cpp
new file mode 100644
--- /dev/null
+++ b/Engine/Plugins/VirtualProduction/DMX/DMXPixelMapping/Source/DMXPixelMappingRenderer/Private/PixelMapRenderElement.cpp
@@ -0,0 +1,39 @@
+// Copyright Epic Games, Inc. All Rights Reserved.
+
+#include "PixelMapRenderElement.h"
+
+
+namespace UE::DMXPixelMapping::Rendering
+{
+	void FPixelMapRenderElement::SetParameters(FPixelMapRenderElementParameters InParameters)
+	{
+#if !defined(_MSC_VER) 			// Atomic for MSVC only
+		const FScopeLock LockParameters(&AccessParametersMutex);
+#endif
+		Parameters = InParameters;
+	}
+
+	FPixelMapRenderElementParameters FPixelMapRenderElement::GetParameters() const
+	{
+#if !defined(_MSC_VER)  		// Atomic for MSVC only
+		const FScopeLock LockParameters(&AccessParametersMutex);
+#endif
+		return Parameters;
+	}
+
+	void FPixelMapRenderElement::SetColor(const FLinearColor& InColor)
+	{
+#if !defined(_MSC_VER)  		// Atomic for MSVC only
+		const FScopeLock LockColorMutex(&AccessColorMutex);
+#endif
+		Color = InColor;
+	}
+
+	FLinearColor FPixelMapRenderElement::GetColor() const
+	{
+#if !defined(_MSC_VER)  		// Atomic for MSVC only
+		const FScopeLock LockColorMutex(&AccessColorMutex);
+#endif
+		return Color;
+	}
+}
diff --git a/Engine/Plugins/VirtualProduction/DMX/DMXPixelMapping/Source/DMXPixelMappingRenderer/Private/PreprocessRenderInputTextureProxy.h b/Engine/Plugins/VirtualProduction/DMX/DMXPixelMapping/Source/DMXPixelMappingRenderer/Private/PreprocessRenderInputTextureProxy.h
--- a/Engine/Plugins/VirtualProduction/DMX/DMXPixelMapping/Source/DMXPixelMappingRenderer/Private/PreprocessRenderInputTextureProxy.h
+++ b/Engine/Plugins/VirtualProduction/DMX/DMXPixelMapping/Source/DMXPixelMappingRenderer/Private/PreprocessRenderInputTextureProxy.h
@@ -6,7 +6,6 @@
 
 #include "UObject/WeakObjectPtr.h"
 
-class IDMXPixelMappingRenderer;
 class UDMXPixelMappingPreprocessRenderer;
 class UTexture;
 
diff --git a/Engine/Plugins/VirtualProduction/DMX/DMXPixelMapping/Source/DMXPixelMappingRenderer/Private/RenderPixelMapProxy.cpp b/Engine/Plugins/VirtualProduction/DMX/DMXPixelMapping/Source/DMXPixelMappingRenderer/Private/RenderPixelMapProxy.cpp
new file mode 100644
--- /dev/null
+++ b/Engine/Plugins/VirtualProduction/DMX/DMXPixelMapping/Source/DMXPixelMappingRenderer/Private/RenderPixelMapProxy.cpp
@@ -0,0 +1,206 @@
+// Copyright Epic Games, Inc. All Rights Reserved.
+
+#include "RenderPixelMapProxy.h"
+
+#include "Engine/TextureRenderTarget2D.h"
+#include "PixelMapRenderElement.h"
+#include "PixelShaderUtils.h"
+#include "ProfilingDebugging/RealtimeGPUProfiler.h"
+#include "RenderingThread.h"
+#include "RHIResources.h"
+#include "TextureResource.h"
+
+
+namespace UE::DMXPixelMapping::Rendering::Private
+{
+	static constexpr auto RenderPassName = TEXT("RenderPixelMapping");
+	static constexpr auto RenderPassHint = TEXT("Render Pixel Mapping");
+};
+
+DECLARE_GPU_STAT_NAMED(DMXPixelMappingShadersStat, UE::DMXPixelMapping::Rendering::Private::RenderPassHint);
+
+class FDMXPixelBlendingQualityDimension : SHADER_PERMUTATION_ENUM_CLASS("PIXELBLENDING_QUALITY", UE::DMXPixelMapping::Rendering::ECellBlendingQuality);
+class FDMXVertexUVDimension : SHADER_PERMUTATION_BOOL("VERTEX_UV_STATIC_CALCULATION");
+
+/**
+ * Pixel Mapping downsampling vertex shader
+ */
+class FDMXPixelMappingRendererVS
+	: public FGlobalShader
+{
+public:
+	DECLARE_GLOBAL_SHADER(FDMXPixelMappingRendererVS);
+	SHADER_USE_PARAMETER_STRUCT(FDMXPixelMappingRendererVS, FGlobalShader);
+
+	using FPermutationDomain = TShaderPermutationDomain<FDMXPixelBlendingQualityDimension, FDMXVertexUVDimension>;
+	BEGIN_SHADER_PARAMETER_STRUCT(FParameters, )
+		SHADER_PARAMETER(FVector4f, DrawRectanglePosScaleBias)
+		SHADER_PARAMETER(FVector4f, DrawRectangleInvTargetSizeAndTextureSize)
+		SHADER_PARAMETER(FVector4f, DrawRectangleUVScaleBias)
+	END_SHADER_PARAMETER_STRUCT()
+
+	static bool ShouldCompilePermutation(const FGlobalShaderPermutationParameters& Parameters) { return true; }
+};
+
+/**
+ * Pixel Mapping downsampling pixel shader
+ */
+class FDMXPixelMappingRendererPS
+	: public FGlobalShader
+{
+public:
+	DECLARE_GLOBAL_SHADER(FDMXPixelMappingRendererPS);
+	SHADER_USE_PARAMETER_STRUCT(FDMXPixelMappingRendererPS, FGlobalShader);
+
+	using FPermutationDomain = TShaderPermutationDomain<FDMXPixelBlendingQualityDimension, FDMXVertexUVDimension>;
+	BEGIN_SHADER_PARAMETER_STRUCT(FParameters, )
+		SHADER_PARAMETER_TEXTURE(Texture2D, InputTexture)
+		SHADER_PARAMETER_SAMPLER(SamplerState, InputSampler)
+
+		SHADER_PARAMETER(float, Brightness)
+		SHADER_PARAMETER(FIntPoint, InputTextureSize)
+		SHADER_PARAMETER(FIntPoint, OutputTextureSize)
+		SHADER_PARAMETER(FVector2f, UVCellSize)
+	END_SHADER_PARAMETER_STRUCT()
+
+	static bool ShouldCompilePermutation(const FGlobalShaderPermutationParameters& Parameters) { return true; }
+};
+
+IMPLEMENT_GLOBAL_SHADER(FDMXPixelMappingRendererVS, "/Plugin/DMXPixelMapping/Private/DMXPixelMapping.usf", "DMXPixelMappingVS", SF_Vertex);
+IMPLEMENT_GLOBAL_SHADER(FDMXPixelMappingRendererPS, "/Plugin/DMXPixelMapping/Private/DMXPixelMapping.usf", "DMXPixelMappingPS", SF_Pixel);
+
+
+
+namespace UE::DMXPixelMapping::Rendering::Private
+{
+	void FRenderPixelMapProxy::Render(
+		UTexture* InInputTexture,
+		UTextureRenderTarget2D* InRenderTarget,
+		const TArray<TSharedRef<FPixelMapRenderElement>>& InElements,
+		float Brightness
+	)
+	{
+		check(IsInGameThread());
+
+		// Ensure lifetime
+		InputTexture = InInputTexture;
+		RenderTarget = InRenderTarget;
+
+		// Actual rendering
+		ENQUEUE_RENDER_COMMAND(RenderTargetResource)(
+			[SharedThis = StaticCastSharedRef<FRenderPixelMapProxy>(AsShared()), this, Elements = InElements, Brightness]
+			(FRHICommandListImmediate& RHICmdList)
+			{				
+				SCOPED_GPU_STAT(RHICmdList, DMXPixelMappingShadersStat);
+				SCOPED_DRAW_EVENTF(RHICmdList, DMXPixelMappingShadersStat, RenderPassName);
+
+				const FTextureResource* InputTextureResource = InputTexture ? InputTexture->GetResource() : nullptr;
+				const FTextureRenderTargetResource* RenderTargetResource = SharedThis->RenderTarget ? SharedThis->RenderTarget->GetRenderTargetResource() : nullptr;
+
+				if (!InputTextureResource || !RenderTargetResource)
+				{
+					return;
+				}
+
+				const FIntPoint PixelSize(1);
+				const FIntPoint TextureSize(1);
+
+				const FIntPoint InputTextureSize = FIntPoint(InputTextureResource->GetSizeX(), InputTextureResource->GetSizeY());
+				const FIntPoint OutputTextureSize = FIntPoint(RenderTargetResource->GetSizeX(), RenderTargetResource->GetSizeY());
+
+				const FTextureRHIRef InputTextureRHI = InputTextureResource->TextureRHI;
+
+				RHICmdList.Transition(FRHITransitionInfo(RenderTargetResource->TextureRHI, ERHIAccess::SRVMask, ERHIAccess::RTV));
+
+				FRHIRenderPassInfo RpInfo(RenderTargetResource->TextureRHI, ERenderTargetActions::Load_Store);
+				RHICmdList.BeginRenderPass(RpInfo, RenderPassName);
+				{
+					RHICmdList.SetViewport(0.f, 0.f, 0.f, OutputTextureSize.X, OutputTextureSize.Y, 1.f);
+
+					for (int32 PixelIndex = 0; PixelIndex < Elements.Num(); PixelIndex++)
+					{
+						const TSharedRef<FPixelMapRenderElement>& Element = Elements[PixelIndex];
+
+						const FVector2D& UV = Element->GetParameters().UV;
+						const FVector2D& UVSize = Element->GetParameters().UVSize;
+						const FVector2D& UVCellSize = Element->GetParameters().UVCellSize;
+
+						const int32 RenderTargetPoisitionX = PixelIndex % OutputTextureSize.X;
+						const int32 RenderTargetPositionY = PixelIndex / OutputTextureSize.X;
+						
+						// Create shader permutations
+						FDMXPixelMappingRendererPS::FPermutationDomain PermutationVector;
+						PermutationVector.Set<FDMXPixelBlendingQualityDimension>(Element->GetParameters().CellBlendingQuality);
+						PermutationVector.Set<FDMXVertexUVDimension>(Element->GetParameters().bStaticCalculateUV);
+
+						// Create shaders
+						FGlobalShaderMap* ShaderMap = GetGlobalShaderMap(GMaxRHIFeatureLevel);
+						TShaderMapRef<FDMXPixelMappingRendererVS> VertexShader(ShaderMap, PermutationVector);
+						TShaderMapRef<FDMXPixelMappingRendererPS> PixelShader(ShaderMap, PermutationVector);
+
+						// Setup graphics pipeline
+						FGraphicsPipelineStateInitializer GraphicsPSOInit;
+						RHICmdList.ApplyCachedRenderTargets(GraphicsPSOInit);
+						GraphicsPSOInit.BlendState = TStaticBlendState<>::GetRHI();
+						GraphicsPSOInit.RasterizerState = TStaticRasterizerState<>::GetRHI();
+						GraphicsPSOInit.DepthStencilState = TStaticDepthStencilState<false, CF_Never>::GetRHI();
+						GraphicsPSOInit.BoundShaderState.VertexDeclarationRHI = GFilterVertexDeclaration.VertexDeclarationRHI;
+						GraphicsPSOInit.BoundShaderState.VertexShaderRHI = VertexShader.GetVertexShader();
+						GraphicsPSOInit.BoundShaderState.PixelShaderRHI = PixelShader.GetPixelShader();
+						GraphicsPSOInit.PrimitiveType = PT_TriangleList;
+						SetGraphicsPipelineState(RHICmdList, GraphicsPSOInit, 0);
+
+						// Set vertex shader buffer
+						FDMXPixelMappingRendererVS::FParameters VSParameters;
+						VSParameters.DrawRectanglePosScaleBias = FVector4f(PixelSize.X, PixelSize.Y, RenderTargetPoisitionX, RenderTargetPositionY);
+						VSParameters.DrawRectangleUVScaleBias = FVector4f(UVSize.X, UVSize.Y, UV.X, UV.Y);
+						VSParameters.DrawRectangleInvTargetSizeAndTextureSize = FVector4f(
+							1.0 / OutputTextureSize.X, 1.0 / OutputTextureSize.Y,
+							1.0 / TextureSize.X, 1.0 / TextureSize.Y);
+						SetShaderParameters(RHICmdList, VertexShader, VertexShader.GetVertexShader(), VSParameters);
+
+						// Set pixel shader buffer
+						FDMXPixelMappingRendererPS::FParameters PSParameters;
+						PSParameters.InputTexture = InputTextureRHI;
+						PSParameters.Brightness = Brightness;
+						PSParameters.InputSampler = TStaticSamplerState<SF_Bilinear, AM_Clamp, AM_Clamp, AM_Clamp>::GetRHI();
+						PSParameters.InputTextureSize = InputTextureSize;
+						PSParameters.OutputTextureSize = OutputTextureSize;
+						PSParameters.UVCellSize = FVector2f(UVCellSize);
+						SetShaderParameters(RHICmdList, PixelShader, PixelShader.GetPixelShader(), PSParameters);
+
+						// Draw a two triangle on the entire viewport.
+						FPixelShaderUtils::DrawFullscreenQuad(RHICmdList, 1);
+					}
+				}
+				RHICmdList.EndRenderPass();
+
+				// Copy texture from GPU to CPU
+				{
+					// Read the contents of a texture to an output CPU buffer
+					TArray<FLinearColor> ColorArray;
+					const FIntRect Rect(0, 0, OutputTextureSize.X, OutputTextureSize.Y);
+
+					// Read surface without flush rendering thread
+					RHICmdList.ReadSurfaceData(RenderTargetResource->TextureRHI, Rect, ColorArray, FReadSurfaceDataFlags());
+
+					// Shrink color array to match num elements
+					if (!ensureMsgf(ColorArray.Num() >= Elements.Num(), TEXT("Pixel maparray is smaller than the number of elements that should have been rendered. Failed to Pixel Map")))
+					{
+						return;
+					}
+
+					for (int32 PixelIndex = 0; PixelIndex < Elements.Num(); PixelIndex++)
+					{
+						Elements[PixelIndex]->SetColor(ColorArray[PixelIndex]);
+					}
+				}
+			});
+	}
+
+	void FRenderPixelMapProxy::AddReferencedObjects(FReferenceCollector& Collector)
+	{
+		Collector.AddReferencedObject(InputTexture);
+		Collector.AddReferencedObject(RenderTarget);
+	}
+}
diff --git a/Engine/Plugins/VirtualProduction/DMX/DMXPixelMapping/Source/DMXPixelMappingRenderer/Private/RenderPixelMapProxy.h b/Engine/Plugins/VirtualProduction/DMX/DMXPixelMapping/Source/DMXPixelMappingRenderer/Private/RenderPixelMapProxy.h
new file mode 100644
--- /dev/null
+++ b/Engine/Plugins/VirtualProduction/DMX/DMXPixelMapping/Source/DMXPixelMappingRenderer/Private/RenderPixelMapProxy.h
@@ -0,0 +1,50 @@
+// Copyright Epic Games, Inc. All Rights Reserved.
+
+#pragma once
+
+#include "DMXPixelMappingPixelMapRenderer.h"
+#include "UObject/GCObject.h"
+
+class UTexture;
+class UTextureRenderTarget2D;
+namespace UE::DMXPixelMapping::Rendering { class FPixelMapRenderElement; }
+
+
+namespace UE::DMXPixelMapping::Rendering::Private
+{
+	/** Interface for the object responsible to render the input texture/material/umg widget. */
+	class FRenderPixelMapProxy
+		: public IRenderPixelMapProxy
+		, public FGCObject
+	{
+	public:
+		/**
+		 * Pixelmapping specific, downsample and draw input texture to destination texture.
+		 *
+		 * @param InInputTexture				The input texture to pixel map
+		 * @param InRenderTarget				The render target that is pixel mapped to
+		 * @param InElements					Elements to be rendered
+		 * @param InBrightness					Global brightness
+		 * @param InCallback					Callback for reading  the pixels from GPU to CPU
+		 */
+		virtual void Render(
+			UTexture* InInputTexture,
+			UTextureRenderTarget2D* InRenderTarget,
+			const TArray<TSharedRef<FPixelMapRenderElement>>& InElements,
+			float InBrightness
+		) override;
+
+	protected:
+		//~ Begin FGCObject interface
+		virtual void AddReferencedObjects(FReferenceCollector& Collector) override;
+		virtual FString GetReferencerName() const override { return TEXT("FRenderPixelMapProxy"); }
+		//~ End FGCObject interface
+
+	private:
+		/** Strong ref to the current input texture */
+		TObjectPtr<UTexture> InputTexture;
+
+		/** Strong ref to the current render target */
+		TObjectPtr<UTextureRenderTarget2D> RenderTarget;
+	};
+}
diff --git a/Engine/Plugins/VirtualProduction/DMX/DMXPixelMapping/Source/DMXPixelMappingRenderer/Public/DMXPixelMappingPixelMapRenderer.h b/Engine/Plugins/VirtualProduction/DMX/DMXPixelMapping/Source/DMXPixelMappingRenderer/Public/DMXPixelMappingPixelMapRenderer.h
new file mode 100644
--- /dev/null
+++ b/Engine/Plugins/VirtualProduction/DMX/DMXPixelMapping/Source/DMXPixelMappingRenderer/Public/DMXPixelMappingPixelMapRenderer.h
@@ -0,0 +1,73 @@
+// Copyright Epic Games, Inc. All Rights Reserved.
+
+#pragma once
+
+#include "UObject/Object.h"
+
+#include "DMXPixelMappingPixelMapRenderer.generated.h"
+
+class FTextureResource;
+class FTextureRenderTargetResource;
+class UTextureRenderTarget2D;
+class UMaterialInstanceDynamic;
+class UMaterialInterface;
+class UTexture;
+class UUserWidget;
+namespace UE::DMXPixelMapping::Rendering { class FPixelMapRenderElement; }
+
+
+namespace UE::DMXPixelMapping::Rendering::Private
+{
+	/** Interface for the object responsible to render the pixel map */
+	class IRenderPixelMapProxy
+		: public TSharedFromThis<IRenderPixelMapProxy>
+	{
+	public:
+		virtual ~IRenderPixelMapProxy() {}
+
+		/**
+		 * Pixelmapping specific, downsample and draw input texture to destination texture.
+		 *
+		 * @param InputTexture					The input texture to render
+		 * @param RenderTarget					The render target to render to
+		 * @param InElements						The elements that are pixel mapped
+		 * @param Brightness						Global brightness of the rendered colors
+		 */
+		virtual void Render(
+			UTexture* InInputTexture,
+			UTextureRenderTarget2D* InRenderTarget,
+			const TArray<TSharedRef<UE::DMXPixelMapping::Rendering::FPixelMapRenderElement>>& InElements,
+			float InBrightness
+		) = 0;
+	};
+}
+
+
+/** Maps pixel map render elements to pixels */
+UCLASS()
+class DMXPIXELMAPPINGRENDERER_API UDMXPixelMappingPixelMapRenderer
+	: public UObject
+{
+	GENERATED_BODY()
+
+public:
+	UDMXPixelMappingPixelMapRenderer();
+
+	/** Sets the elements to render */
+	void SetElements(const TArray<TSharedRef<UE::DMXPixelMapping::Rendering::FPixelMapRenderElement>>& InElements);
+
+	/** Renders the current elements */
+	void Render(UTexture* InputTexture, float Brightness);
+
+private:
+	static const FIntPoint MaxPixelMapSize;
+
+	UPROPERTY(Transient, DuplicateTransient)
+	TObjectPtr<UTextureRenderTarget2D> PixelMapRenderTarget;
+
+	/** The elements to render */
+	TArray<TSharedRef<UE::DMXPixelMapping::Rendering::FPixelMapRenderElement>> RenderElements;
+
+	/** Proxy that carries out the actual rendering */
+	TSharedPtr<UE::DMXPixelMapping::Rendering::Private::IRenderPixelMapProxy> RendererPixelMapProxy;
+};
diff --git a/Engine/Plugins/VirtualProduction/DMX/DMXPixelMapping/Source/DMXPixelMappingRenderer/Public/DMXPixelMappingPixelMapRendererParams.h b/Engine/Plugins/VirtualProduction/DMX/DMXPixelMapping/Source/DMXPixelMappingRenderer/Public/DMXPixelMappingPixelMapRendererParams.h
new file mode 100644
--- /dev/null
+++ b/Engine/Plugins/VirtualProduction/DMX/DMXPixelMapping/Source/DMXPixelMappingRenderer/Public/DMXPixelMappingPixelMapRendererParams.h
@@ -0,0 +1,45 @@
+// Copyright Epic Games, Inc. All Rights Reserved.
+
+#pragma once
+
+#include "HAL/Platform.h"
+#include "Math/IntPoint.h"
+#include "Math/Vector2D.h"
+
+
+namespace UE::DMXPixelMapping::Rendering::PixelMap::Private
+{
+	/**
+	 * Used in shader permutation for determining number of samples to use in texture blending.
+	 * If adding to this you must also adjust the public facing option: 'EPixelBlendingQuality' under the runtime module's DMXPixelMappingOutputComponent.h
+	 */
+	enum class ECellBlendingQuality : uint8
+	{
+		Low,
+		Medium,
+		High,
+
+		MAX
+	};
+
+	struct FPixelMapRendererParams
+	{
+		/** Position in screen pixels of the top left corner of the quad */
+		FIntPoint Position;
+
+		/** Position in texels of the top left corner of the quad's UV's */
+		FVector2D UV;
+
+		/** Size in texels of the quad's total UV space */
+		FVector2D UVSize;
+
+		/** Size in texels of UV.May match UVSize */
+		FVector2D UVCellSize;
+
+		/** The quality of color samples in the pixel shader(number of samples) */
+		ECellBlendingQuality CellBlendingQuality;
+
+		/** Calculates the UV point to sample purely on the UV position/size. Works best for renderers which represent a single pixel */
+		bool bStaticCalculateUV;
+	};
+}
diff --git a/Engine/Plugins/VirtualProduction/DMX/DMXPixelMapping/Source/DMXPixelMappingRenderer/Public/IDMXPixelMappingRenderer.h b/Engine/Plugins/VirtualProduction/DMX/DMXPixelMapping/Source/DMXPixelMappingRenderer/Public/IDMXPixelMappingRenderer.h
--- a/Engine/Plugins/VirtualProduction/DMX/DMXPixelMapping/Source/DMXPixelMappingRenderer/Public/IDMXPixelMappingRenderer.h
+++ b/Engine/Plugins/VirtualProduction/DMX/DMXPixelMapping/Source/DMXPixelMappingRenderer/Public/IDMXPixelMappingRenderer.h
@@ -21,6 +21,7 @@ class UUserWidget;
  * Used in shader permutation for determining number of samples to use in texture blending.
  * If adding to this you must also adjust the public facing option: 'EPixelBlendingQuality' under the runtime module's DMXPixelMappingOutputComponent.h
  */
+enum class UE_DEPRECATED(5.3, "Please use UDMXPixelMappingPixelMapRenderer instead.") EDMXPixelShaderBlendingQuality : uint8;
 enum class EDMXPixelShaderBlendingQuality : uint8
 {
 	Low,
@@ -35,6 +36,7 @@ enum class EDMXPixelShaderBlendingQuality : uint8
  * Downsample pixel preview rendering params.
  * Using for pixel rendering setting in preview
  */
+struct UE_DEPRECATED(5.3, "Deprecated in favor of FPixelMapRenderElement.") FDMXPixelMappingDownsamplePixelParam;
 struct FDMXPixelMappingDownsamplePixelPreviewParam
 {
 	/** Position in screen pixels of the top left corner of the quad */
@@ -79,6 +81,7 @@ struct FDMXPixelMappingDownsamplePixelParam
 	bool bStaticCalculateUV;
 };
 
+struct UE_DEPRECATED(5.3, "Deprecated in favor of FPixelMapRenderElement.") FDMXPixelMappingDownsamplePixelParamsV2;
 struct FDMXPixelMappingDownsamplePixelParamsV2
 {
 	/** Position in screen pixels of the top left corner of the quad */
@@ -100,22 +103,11 @@ struct FDMXPixelMappingDownsamplePixelParamsV2
 	bool bStaticCalculateUV;
 };
 
-struct FDMXPixelMappingRenderTextureParams
-{
-	int32 DownsampleTexture = 8;
-
-	int32 NumDownSamplePasses = 1;
-	float Distance = 0.2;
-	int32 DistanceSteps = 1;
-	int32 RadialSteps = 1;
-	float RadialOffset = 1;
-	int32 KernelPower = 5;
-};
-
 
 /**
  * The public interface of the Pixel Mapping renderer instance interface.
  */
+PRAGMA_DISABLE_DEPRECATION_WARNINGS
 class IDMXPixelMappingRenderer 
 	: public TSharedFromThis<IDMXPixelMappingRenderer>
 {
@@ -134,7 +126,8 @@ public:
 	 * @param InDownsamplePixelPass			Pixels rendering params
 	 * @param InCallback					Callback for reading  the pixels from GPU to CPU			
 	 */
-	 virtual void DownsampleRender(
+	UE_DEPRECATED(5.3, "Please use UDMXDMXPixelMappingPixelMapRenderer instead.")
+	virtual void DownsampleRender(
 		const FTextureResource* InputTexture,
 		const FTextureResource* DstTexture,
 		const FTextureRenderTargetResource* DstTextureTargetResource,
@@ -148,7 +141,7 @@ public:
 	 * @param InRenderTarget				2D render target texture resource
 	 * @param InMaterialInterface			Material to use
 	 */
-	 UE_DEPRECATED(5.3, "Please use FWidgetRenderer::DrawWidget or similar engine methods instead.")
+	UE_DEPRECATED(5.3, "Please use FWidgetRenderer::DrawWidget or similar engine methods instead.")
 	virtual void RenderMaterial(UTextureRenderTarget2D* InRenderTarget, UMaterialInterface* InMaterialInterface) const = 0;
 
 	/**
@@ -157,7 +150,7 @@ public:
 	 * @param InRenderTarget				2D render target texture resource
 	 * @param InUserWidget					UMG widget to use
 	 */
-	 UE_DEPRECATED(5.3, "Please use FWidgetRenderer::DrawWidget or similar engine methods instead.")
+	UE_DEPRECATED(5.3, "Please use FWidgetRenderer::DrawWidget or similar engine methods instead.")
 	virtual void RenderWidget(UTextureRenderTarget2D* InRenderTarget, UUserWidget* InUserWidget) const  = 0;
 
 	/**
@@ -168,6 +161,7 @@ public:
 	 * @param InSize						Rendering size
 	 * @param bSRGBSource					If the source texture is sRGB
 	 */
+	UE_DEPRECATED(5.3, "Removed without replacement. Please use core engine methods to draw texture to rectangle.")
 	virtual void RenderTextureToRectangle(const FTextureResource* InTextureResource, const FTexture2DRHIRef InRenderTargetTexture, FVector2D InSize, bool bSRGBSource) const = 0;
 
 #if WITH_EDITOR
@@ -178,15 +172,18 @@ public:
 	 * @param DownsampleResource			Rendering resource of RenderTarget texture
 	 * @param InPixelPreviewParamSet		Pixels rendering params
 	 */
-	virtual void RenderPreview(const FTextureResource* TextureResource, const FTextureResource* DownsampleResource, TArray<FDMXPixelMappingDownsamplePixelPreviewParam>&& InPixelPreviewParamSet) const = 0;
+	 UE_DEPRECATED(5.3, "Pixel mapping no longer has a dedicated preview renderer. Instead draw pixel map elements to a view port (see FDMXPixelMappingPreviewViewportClient.")
+	 virtual void RenderPreview(const FTextureResource* TextureResource, const FTextureResource* DownsampleResource, TArray<FDMXPixelMappingDownsamplePixelPreviewParam>&& InPixelPreviewParamSet) const = 0;
 #endif // WITH_EDITOR
 
 	/**
 	* Sets the brigthness of the renderer
 	*/
+	UE_DEPRECATED(5.3, "IDMXPixelMappingRenderer and all its members are deprecated.")
 	void SetBrightness(const float InBrightness) { Brightness = InBrightness; }
 
 protected:
 	/** Brightness multiplier for the renderer */
 	float Brightness = 1.0f;
 };
+PRAGMA_ENABLE_DEPRECATION_WARNINGS
diff --git a/Engine/Plugins/VirtualProduction/DMX/DMXPixelMapping/Source/DMXPixelMappingRenderer/Public/PixelMapRenderElement.h b/Engine/Plugins/VirtualProduction/DMX/DMXPixelMapping/Source/DMXPixelMappingRenderer/Public/PixelMapRenderElement.h
new file mode 100644
--- /dev/null
+++ b/Engine/Plugins/VirtualProduction/DMX/DMXPixelMapping/Source/DMXPixelMappingRenderer/Public/PixelMapRenderElement.h
@@ -0,0 +1,93 @@
+// Copyright Epic Games, Inc. All Rights Reserved.
+
+#pragma once
+
+#if defined(_MSC_VER) 
+#include <atomic>
+#else
+#include "HAL/CriticalSection.h"
+#endif
+#include "Math/Color.h"
+#include "Math/Vector2D.h"
+
+
+namespace UE::DMXPixelMapping::Rendering
+{
+	/**
+	 * Used in shader permutation for determining number of samples to use in texture blending.
+	 * If adding to this you must also adjust the public facing option: 'EPixelBlendingQuality' under the runtime module's DMXPixelMappingOutputComponent.h
+	 */
+	enum class DMXPIXELMAPPINGRENDERER_API ECellBlendingQuality : uint8
+	{
+		Low,
+		Medium,
+		High,
+		
+		MAX
+	};
+
+	struct DMXPIXELMAPPINGRENDERER_API FPixelMapRenderElementParameters
+	{
+		/** Position in texels of the top left corner of the quad's UV's */
+		FVector2D UV;
+
+		/** Size in texels of the quad's total UV space */
+		FVector2D UVSize;
+
+		/** Size in texels of UV.May match UVSize */
+		FVector2D UVCellSize;
+
+		/** The quality of color samples in the pixel shader(number of samples) */
+		ECellBlendingQuality CellBlendingQuality;
+
+		/** Calculates the UV point to sample purely on the UV position/size. Works best for renderers which represent a single pixel */
+		bool bStaticCalculateUV = true;
+	};
+
+	/** An element rendered by the pixel map renderer */
+	class DMXPIXELMAPPINGRENDERER_API FPixelMapRenderElement
+		: public TSharedFromThis<FPixelMapRenderElement>
+	{
+	public:
+		FPixelMapRenderElement(FPixelMapRenderElementParameters InParameters)
+			: Parameters(InParameters)
+		{}
+
+		virtual ~FPixelMapRenderElement() {};
+
+		/** Sets the parameters, thread-safe */
+		void SetParameters(FPixelMapRenderElementParameters InParameters);
+
+		/** Gets the parameters, thread-safe */
+		FPixelMapRenderElementParameters GetParameters() const;
+
+		/** Sets the pixel color, thread safe. This is called from the renderer to set the rendered color. */
+		void SetColor(const FLinearColor& Color);
+
+		/** Gets the pixel color, thread safe */
+		FLinearColor GetColor() const;
+
+	private:
+		// 5.3: Significantly better performance with atomics on MSVC. 
+		// However not all compilers currently support non-lockfree atomics.
+#if defined(_MSC_VER) 
+		/** The current parameters */
+		std::atomic<FPixelMapRenderElementParameters> Parameters;
+
+		/** The current pixel color */
+		std::atomic<FLinearColor> Color;
+#else
+		/** The current parameters */
+		FPixelMapRenderElementParameters Parameters;
+
+		/** The current pixel color */
+		FLinearColor Color;
+
+		/** Mutex acccess to parameters */
+		mutable FCriticalSection AccessParametersMutex;
+
+		/** Mutex acccess to parameters */
+		mutable FCriticalSection AccessColorMutex;
+#endif
+	};
+}
diff --git a/Engine/Plugins/VirtualProduction/DMX/DMXPixelMapping/Source/DMXPixelMappingRuntime/Private/Components/DMXPixelMappingBaseComponent.cpp b/Engine/Plugins/VirtualProduction/DMX/DMXPixelMapping/Source/DMXPixelMappingRuntime/Private/Components/DMXPixelMappingBaseComponent.cpp
--- a/Engine/Plugins/VirtualProduction/DMX/DMXPixelMapping/Source/DMXPixelMappingRuntime/Private/Components/DMXPixelMappingBaseComponent.cpp
+++ b/Engine/Plugins/VirtualProduction/DMX/DMXPixelMapping/Source/DMXPixelMappingRuntime/Private/Components/DMXPixelMappingBaseComponent.cpp
@@ -13,8 +13,6 @@ UDMXPixelMappingBaseComponent::FDMXPixelMappingOnComponentAdded UDMXPixelMapping
 UDMXPixelMappingBaseComponent::FDMXPixelMappingOnComponentRemoved UDMXPixelMappingBaseComponent::OnComponentRemoved;
 UDMXPixelMappingBaseComponent::FDMXPixelMappingOnComponentRenamed UDMXPixelMappingBaseComponent::OnComponentRenamed;
 
-UDMXPixelMappingBaseComponent::UDMXPixelMappingBaseComponent()
-{}
 
 UDMXPixelMappingBaseComponent::FDMXPixelMappingOnComponentAdded& UDMXPixelMappingBaseComponent::UDMXPixelMappingBaseComponent::GetOnComponentAdded()
 {
@@ -82,11 +80,6 @@ const FName& UDMXPixelMappingBaseComponent::GetNamePrefix()
 	return NamePrefix;
 }
 
-TStatId UDMXPixelMappingBaseComponent::GetStatId() const
-{
-	RETURN_QUICK_DECLARE_CYCLE_STAT(UDMXPixelMappingBaseComponent, STATGROUP_Tickables);
-}
-
 int32 UDMXPixelMappingBaseComponent::GetChildrenCount() const
 {
 	return Children.Num();
diff --git a/Engine/Plugins/VirtualProduction/DMX/DMXPixelMapping/Source/DMXPixelMappingRuntime/Private/Components/DMXPixelMappingFixtureGroupComponent.cpp b/Engine/Plugins/VirtualProduction/DMX/DMXPixelMapping/Source/DMXPixelMappingRuntime/Private/Components/DMXPixelMappingFixtureGroupComponent.cpp
--- a/Engine/Plugins/VirtualProduction/DMX/DMXPixelMapping/Source/DMXPixelMappingRuntime/Private/Components/DMXPixelMappingFixtureGroupComponent.cpp
+++ b/Engine/Plugins/VirtualProduction/DMX/DMXPixelMapping/Source/DMXPixelMappingRuntime/Private/Components/DMXPixelMappingFixtureGroupComponent.cpp
@@ -143,12 +143,14 @@ void UDMXPixelMappingFixtureGroupComponent::SendDMX()
 
 void UDMXPixelMappingFixtureGroupComponent::QueueDownsample()
 {
+	PRAGMA_DISABLE_DEPRECATION_WARNINGS
 	ForEachChild([&](UDMXPixelMappingBaseComponent* InComponent) {
 		if (UDMXPixelMappingOutputComponent* Component = Cast<UDMXPixelMappingOutputComponent>(InComponent))
 		{
 			Component->QueueDownsample();
 		}
 	}, false);
+	PRAGMA_ENABLE_DEPRECATION_WARNINGS
 }
 
 void UDMXPixelMappingFixtureGroupComponent::SetPosition(const FVector2D& NewPosition)
diff --git a/Engine/Plugins/VirtualProduction/DMX/DMXPixelMapping/Source/DMXPixelMappingRuntime/Private/Components/DMXPixelMappingFixtureGroupItemComponent.cpp b/Engine/Plugins/VirtualProduction/DMX/DMXPixelMapping/Source/DMXPixelMappingRuntime/Private/Components/DMXPixelMappingFixtureGroupItemComponent.cpp
--- a/Engine/Plugins/VirtualProduction/DMX/DMXPixelMapping/Source/DMXPixelMappingRuntime/Private/Components/DMXPixelMappingFixtureGroupItemComponent.cpp
+++ b/Engine/Plugins/VirtualProduction/DMX/DMXPixelMapping/Source/DMXPixelMappingRuntime/Private/Components/DMXPixelMappingFixtureGroupItemComponent.cpp
@@ -6,7 +6,6 @@
 #include "DMXPixelMappingMainStreamObjectVersion.h"
 #include "DMXPixelMappingRuntimeUtils.h"
 #include "DMXPixelMappingTypes.h"
-#include "IDMXPixelMappingRenderer.h"
 #include "ColorSpace/DMXPixelMappingColorSpace_RGBCMY.h"
 #include "Components/DMXPixelMappingFixtureGroupComponent.h"
 #include "Components/DMXPixelMappingRendererComponent.h"
@@ -87,11 +86,31 @@ void UDMXPixelMappingFixtureGroupItemComponent::Serialize(FArchive& Ar)
 		}
 	}
 #endif
+
+	if (UDMXPixelMappingRendererComponent* RendererComponent = GetRendererComponent())
+	{
+		RendererComponent->InvalidatePixelMapRenderer();
+	}
+}
+
+void UDMXPixelMappingFixtureGroupItemComponent::PostInitProperties()
+{
+	Super::PostInitProperties();
+	if (IsTemplate())
+	{
+		return;
+	}
+
+	UpdateRenderElement();
 }
 
 void UDMXPixelMappingFixtureGroupItemComponent::PostLoad()
 {
 	Super::PostLoad();
+	if(IsTemplate())
+	{
+		return;
+	}
 
 	// Add valid modulators to modulator classes, remove invalid modulators
 	for (int32 IndexModulator = 0; Modulators.IsValidIndex(IndexModulator); )
@@ -126,6 +145,10 @@ void UDMXPixelMappingFixtureGroupItemComponent::PostLoad()
 void UDMXPixelMappingFixtureGroupItemComponent::BeginDestroy() 
 {
 	Super::BeginDestroy();
+	if (IsTemplate())
+	{
+		return;
+	}
 
 #if WITH_EDITOR
 	// Set the transient Color Space Class property to the initial value
@@ -134,6 +157,8 @@ void UDMXPixelMappingFixtureGroupItemComponent::BeginDestroy()
 		ColorSpace->GetOnPostEditChangedProperty().RemoveAll(this);
 	}
 #endif // WITH_EDITOR
+
+	PixelMapRenderElement.Reset();
 }
 
 #if WITH_EDITOR
@@ -260,20 +285,19 @@ void UDMXPixelMappingFixtureGroupItemComponent::SendDMX()
 {
 	SCOPE_CYCLE_COUNTER(STAT_DMXPixelMaping_FixtureGroupItem);
 
-	UDMXEntityFixturePatch* FixturePatch = FixturePatchRef.GetFixturePatch();
-	if (!FixturePatch)
+	if (!PixelMapRenderElement.IsValid())
 	{
 		return;
 	}
 
-	UDMXLibrary* Library = FixturePatch->GetParentLibrary();
-	if (!Library)
+	UDMXEntityFixturePatch* FixturePatch = FixturePatchRef.GetFixturePatch();
+	if (!FixturePatch)
 	{
 		return;
 	}
 
-	UDMXPixelMappingRendererComponent* RendererComponent = GetRendererComponent();
-	if (!RendererComponent)
+	UDMXLibrary* Library = FixturePatch->GetParentLibrary();
+	if (!Library)
 	{
 		return;
 	}
@@ -283,14 +307,7 @@ void UDMXPixelMappingFixtureGroupItemComponent::SendDMX()
 		return;
 	}
 
-	// Get the color data from the rendered component
-	FLinearColor PixelColor;
-	if (!RendererComponent->GetDownsampleBufferPixel(DownsamplePixelIndex, PixelColor))
-	{
-		return;
-	}
-
-	ColorSpace->SetRGBA(PixelColor);
+	ColorSpace->SetRGBA(PixelMapRenderElement->GetColor());
 	TMap<FDMXAttributeName, float> AttributeToValueMap = ColorSpace->GetAttributeNameToValueMap();
 
 	for (UDMXModulator* Modulator : Modulators)
@@ -314,14 +331,17 @@ void UDMXPixelMappingFixtureGroupItemComponent::SendDMX()
 
 void UDMXPixelMappingFixtureGroupItemComponent::QueueDownsample()
 {
+	// DEPRECATED 5.3
+
 	// Queue pixels into the downsample rendering
+	PRAGMA_DISABLE_DEPRECATION_WARNINGS
 	UDMXPixelMappingRendererComponent* RendererComponent = GetRendererComponent();
 
 	if (!ensure(RendererComponent))
 	{
 		return;
 	}
-
+	
 	UTexture* InputTexture = RendererComponent->GetRenderedInputTexture();
 	if (!ensure(InputTexture))
 	{
@@ -330,7 +350,7 @@ void UDMXPixelMappingFixtureGroupItemComponent::QueueDownsample()
 
 	// Store pixel position
 	DownsamplePixelIndex = RendererComponent->GetDownsamplePixelNum();
-	
+
 	const uint32 TextureSizeX = InputTexture->GetResource()->GetSizeX();
 	const uint32 TextureSizeY = InputTexture->GetResource()->GetSizeY();
 	check(TextureSizeX > 0 && TextureSizeY > 0);
@@ -351,6 +371,7 @@ void UDMXPixelMappingFixtureGroupItemComponent::QueueDownsample()
 	};
 
 	RendererComponent->AddPixelToDownsampleSet(MoveTemp(DownsamplePixelParams));
+	PRAGMA_ENABLE_DEPRECATION_WARNINGS
 }
 
 void UDMXPixelMappingFixtureGroupItemComponent::SetPosition(const FVector2D& NewPosition)
@@ -365,6 +386,8 @@ void UDMXPixelMappingFixtureGroupItemComponent::SetPosition(const FVector2D& New
 	}
 	PRAGMA_ENABLE_DEPRECATION_WARNINGS
 #endif
+
+	UpdateRenderElement();
 }
 
 void UDMXPixelMappingFixtureGroupItemComponent::SetSize(const FVector2D& NewSize)
@@ -379,6 +402,8 @@ void UDMXPixelMappingFixtureGroupItemComponent::SetSize(const FVector2D& NewSize
 	}
 	PRAGMA_ENABLE_DEPRECATION_WARNINGS
 #endif
+
+	UpdateRenderElement();
 }
 
 bool UDMXPixelMappingFixtureGroupItemComponent::IsOverParent() const
@@ -408,14 +433,15 @@ bool UDMXPixelMappingFixtureGroupItemComponent::IsOverParent() const
 
 void UDMXPixelMappingFixtureGroupItemComponent::RenderWithInputAndSendDMX()
 {
-	if (UDMXPixelMappingRendererComponent* RendererComponent = GetRendererComponent())
-	{
-		RendererComponent->RendererInputTexture();
-	}
-
 	RenderAndSendDMX();
 }
 
+TSharedRef<UE::DMXPixelMapping::Rendering::FPixelMapRenderElement> UDMXPixelMappingFixtureGroupItemComponent::GetOrCreatePixelMapRenderElement()
+{
+	UpdateRenderElement();
+	return PixelMapRenderElement.ToSharedRef();
+}
+
 bool UDMXPixelMappingFixtureGroupItemComponent::CanBeMovedTo(const UDMXPixelMappingBaseComponent* Component) const
 {
 	if (const UDMXPixelMappingFixtureGroupComponent* FixtureGroupComponent = Cast<UDMXPixelMappingFixtureGroupComponent>(Component))
@@ -429,6 +455,29 @@ bool UDMXPixelMappingFixtureGroupItemComponent::CanBeMovedTo(const UDMXPixelMapp
 	return false;
 }
 
+void UDMXPixelMappingFixtureGroupItemComponent::UpdateRenderElement()
+{
+	using namespace UE::DMXPixelMapping::Rendering;
+
+	const UDMXPixelMappingRendererComponent* RendererComponent = GetRendererComponent();
+	const UTexture* InputTexture = RendererComponent ? RendererComponent->GetRenderedInputTexture() : nullptr;
+	const double InputTextureWidth = InputTexture ? InputTexture->GetSurfaceWidth() : 0.0;
+	const double InputTextureHeight = InputTexture ? InputTexture->GetSurfaceHeight() : 0.0;
+
+	FPixelMapRenderElementParameters Parameters;
+	Parameters.UV = FVector2D(GetPosition().X / InputTextureWidth, GetPosition().Y / InputTextureHeight);
+	Parameters.UVSize = FVector2D(GetSize().X / InputTextureWidth, GetSize().Y / InputTextureHeight);
+	Parameters.UVCellSize = Parameters.UVSize / 2.f;
+	Parameters.CellBlendingQuality = static_cast<ECellBlendingQuality>(CellBlendingQuality);
+	Parameters.bStaticCalculateUV = true;
+
+	if (!PixelMapRenderElement.IsValid())
+	{
+		PixelMapRenderElement = MakeShared<FPixelMapRenderElement>(Parameters);
+	}
+	PixelMapRenderElement->SetParameters(Parameters);
+}
+
 #if WITH_EDITOR
 void UDMXPixelMappingFixtureGroupItemComponent::OnColorSpacePostEditChangeProperties(FPropertyChangedEvent& PropertyChangedEvent)
 {
diff --git a/Engine/Plugins/VirtualProduction/DMX/DMXPixelMapping/Source/DMXPixelMappingRuntime/Private/Components/DMXPixelMappingMatrixCellComponent.cpp b/Engine/Plugins/VirtualProduction/DMX/DMXPixelMapping/Source/DMXPixelMappingRuntime/Private/Components/DMXPixelMappingMatrixCellComponent.cpp
--- a/Engine/Plugins/VirtualProduction/DMX/DMXPixelMapping/Source/DMXPixelMappingRuntime/Private/Components/DMXPixelMappingMatrixCellComponent.cpp
+++ b/Engine/Plugins/VirtualProduction/DMX/DMXPixelMapping/Source/DMXPixelMappingRuntime/Private/Components/DMXPixelMappingMatrixCellComponent.cpp
@@ -2,14 +2,17 @@
 
 #include "Components/DMXPixelMappingMatrixCellComponent.h"
 
-#include "DMXPixelMappingTypes.h"
-#include "IDMXPixelMappingRenderer.h"
 #include "ColorSpace/DMXPixelMappingColorSpace.h"
 #include "Components/DMXPixelMappingFixtureGroupComponent.h"
 #include "Components/DMXPixelMappingMatrixComponent.h"
 #include "Components/DMXPixelMappingRendererComponent.h"
+#include "DMXConversions.h"
+#include "DMXPixelMappingTypes.h"
+#include "Engine/Texture.h"
 #include "Library/DMXEntityFixturePatch.h"
 #include "Library/DMXEntityFixtureType.h"
+#include "PixelMapRenderElement.h"
+#include "TextureResource.h"
 
 #if WITH_EDITOR
 #include "DMXPixelMappingComponentWidget.h"
@@ -17,14 +20,10 @@
 #include "SDMXPixelMappingComponentLabel.h"
 #endif // WITH_EDITOR
 
-#include "TextureResource.h"
-#include "Engine/Texture.h"
-
-
-DECLARE_CYCLE_STAT(TEXT("Send Matrix Cell"), STAT_DMXPixelMaping_SendMatrixCell, STATGROUP_DMXPIXELMAPPING);
 
 #define LOCTEXT_NAMESPACE "DMXPixelMappingMatrixPixelComponent"
 
+
 UDMXPixelMappingMatrixCellComponent::UDMXPixelMappingMatrixCellComponent()
 	: DownsamplePixelIndex(0)
 {
@@ -34,12 +33,47 @@ UDMXPixelMappingMatrixCellComponent::UDMXPixelMappingMatrixCellComponent()
 #endif // WITH_EDITORONLY_DATA
 }
 
+void UDMXPixelMappingMatrixCellComponent::PostInitProperties()
+{
+	Super::PostInitProperties();
+	if (IsTemplate())
+	{
+		return;
+	}
+
+	UpdateRenderElement();
+}
+
+void UDMXPixelMappingMatrixCellComponent::PostLoad()
+{
+	Super::PostLoad();
+	if (IsTemplate())
+	{
+		return;
+	}
+
+	UpdateRenderElement();
+}
+
+void UDMXPixelMappingMatrixCellComponent::BeginDestroy()
+{
+	Super::BeginDestroy();
+	if (IsTemplate())
+	{
+		return;
+	}
+
+	PixelMapRenderElement.Reset();
+}
+
 #if WITH_EDITOR
 void UDMXPixelMappingMatrixCellComponent::PostEditChangeProperty(FPropertyChangedEvent& PropertyChangedEvent)
 {
 	// Call the parent at the first place
 	Super::PostEditChangeProperty(PropertyChangedEvent);
 
+	UpdateRenderElement();
+
 	FName PropertyName = PropertyChangedEvent.GetPropertyName();
 
 	PRAGMA_DISABLE_DEPRECATION_WARNINGS
@@ -67,6 +101,11 @@ void UDMXPixelMappingMatrixCellComponent::PostEditChangeProperty(FPropertyChange
 		}
 	}
 	PRAGMA_ENABLE_DEPRECATION_WARNINGS
+
+	if (UDMXPixelMappingRendererComponent* RendererComponent = GetRendererComponent())
+	{
+		RendererComponent->InvalidatePixelMapRenderer();
+	}
 }
 #endif // WITH_EDITOR
 
@@ -193,12 +232,15 @@ const FName& UDMXPixelMappingMatrixCellComponent::GetNamePrefix()
 
 void UDMXPixelMappingMatrixCellComponent::ResetDMX()
 {
-	UDMXPixelMappingRendererComponent* RendererComponent = GetRendererComponent();
-	if (RendererComponent)
+	UDMXPixelMappingMatrixComponent* ParentMatrix = Cast<UDMXPixelMappingMatrixComponent>(GetParent());
+	if (!ParentMatrix)
 	{
-		RendererComponent->ResetColorDownsampleBufferPixel(DownsamplePixelIndex);
+		return;
 	}
 
+	UDMXPixelMappingColorSpace* ColorSpace = ParentMatrix->ColorSpace;
+	ColorSpace->SetRGBA(FLinearColor::Black);
+
 	// No need to send dmx, that is done by the parent matrix
 }
 
@@ -219,7 +261,10 @@ FString UDMXPixelMappingMatrixCellComponent::GetUserFriendlyName() const
 
 void UDMXPixelMappingMatrixCellComponent::QueueDownsample()
 {
+	// DEPRECATED 5.3
+
 	// Queue pixels into the downsample rendering
+	PRAGMA_DISABLE_DEPRECATION_WARNINGS
 	UDMXPixelMappingMatrixComponent* MatrixComponent = Cast<UDMXPixelMappingMatrixComponent>(GetParent());
 	UDMXPixelMappingRendererComponent* RendererComponent = GetRendererComponent();
 
@@ -240,7 +285,9 @@ void UDMXPixelMappingMatrixCellComponent::QueueDownsample()
 	}
 
 	// Store downsample index
+	PRAGMA_DISABLE_DEPRECATION_WARNINGS
 	DownsamplePixelIndex = RendererComponent->GetDownsamplePixelNum();
+	PRAGMA_ENABLE_DEPRECATION_WARNINGS
 
 	const uint32 TextureSizeX = InputTexture->GetResource()->GetSizeX();
 	const uint32 TextureSizeY = InputTexture->GetResource()->GetSizeY();
@@ -262,6 +309,7 @@ void UDMXPixelMappingMatrixCellComponent::QueueDownsample()
 	};
 
 	RendererComponent->AddPixelToDownsampleSet(MoveTemp(DownsamplePixelParam));
+	PRAGMA_ENABLE_DEPRECATION_WARNINGS
 }
 
 void UDMXPixelMappingMatrixCellComponent::SetCellCoordinate(FIntPoint InCellCoordinate)
@@ -269,64 +317,65 @@ void UDMXPixelMappingMatrixCellComponent::SetCellCoordinate(FIntPoint InCellCoor
 	CellCoordinate = InCellCoordinate;
 }
 
+PRAGMA_DISABLE_DEPRECATION_WARNINGS
 void UDMXPixelMappingMatrixCellComponent::RenderWithInputAndSendDMX()
 {
-	if (UDMXPixelMappingRendererComponent* RendererComponent = GetRendererComponent())
-	{
-		RendererComponent->RendererInputTexture();
-	}
-
 	RenderAndSendDMX();
 }
+PRAGMA_ENABLE_DEPRECATION_WARNINGS
 
 bool UDMXPixelMappingMatrixCellComponent::CanBeMovedTo(const UDMXPixelMappingBaseComponent* OtherComponent) const
 {
 	return OtherComponent && OtherComponent == GetParent();
 }
 
-UDMXPixelMappingRendererComponent* UDMXPixelMappingMatrixCellComponent::GetRendererComponent() const
+TSharedRef<UE::DMXPixelMapping::Rendering::FPixelMapRenderElement> UDMXPixelMappingMatrixCellComponent::GetOrCreatePixelMapRenderElement()
 {
-	for (UDMXPixelMappingBaseComponent* ParentComponent = GetParent(); ParentComponent; ParentComponent = ParentComponent->GetParent())
-	{
-		if (UDMXPixelMappingRendererComponent* RendererComponent = Cast<UDMXPixelMappingRendererComponent>(ParentComponent))
-		{
-			return RendererComponent;
-		}
-	}
-	return nullptr;
+	UpdateRenderElement();
+	return PixelMapRenderElement.ToSharedRef();
 }
 
-TMap<FDMXAttributeName, float> UDMXPixelMappingMatrixCellComponent::CreateAttributeValues() const
+void UDMXPixelMappingMatrixCellComponent::UpdateRenderElement()
 {
-	TMap<FDMXAttributeName, float> AttributeToNormalizedValueMap;
+	using namespace UE::DMXPixelMapping::Rendering;
 
-	UDMXPixelMappingMatrixComponent* ParentMatrix = Cast<UDMXPixelMappingMatrixComponent>(GetParent());
-	if (!ParentMatrix)
-	{
-		return AttributeToNormalizedValueMap;
-	}
+	const UDMXPixelMappingRendererComponent* RendererComponent = GetRendererComponent();
+	const UTexture* InputTexture = RendererComponent ? RendererComponent->GetRenderedInputTexture() : nullptr;
+	const double InputTextureWidth = InputTexture ? InputTexture->GetSurfaceWidth() : 0.0;
+	const double InputTextureHeight = InputTexture ? InputTexture->GetSurfaceHeight() : 0.0;
 
-	UDMXPixelMappingColorSpace* ColorSpace = ParentMatrix->ColorSpace;
-	if (!ColorSpace)
+	FPixelMapRenderElementParameters Parameters;
+	Parameters.UV = FVector2D(GetPosition().X / InputTextureWidth, GetPosition().Y / InputTextureHeight);
+	Parameters.UVSize = FVector2D(GetSize().X / InputTextureWidth, GetSize().Y / InputTextureHeight);
+	Parameters.UVCellSize = Parameters.UVSize / 2.f;
+	Parameters.CellBlendingQuality = static_cast<ECellBlendingQuality>(CellBlendingQuality);
+	Parameters.bStaticCalculateUV = true;
+
+	if (!PixelMapRenderElement.IsValid())
 	{
-		return AttributeToNormalizedValueMap;
+		PixelMapRenderElement = MakeShared<FPixelMapRenderElement>(Parameters);
 	}
-
-	UDMXPixelMappingRendererComponent* RendererComponent = GetRendererComponent();
-	if (!RendererComponent)
+	else
 	{
-		return AttributeToNormalizedValueMap;
+		PixelMapRenderElement->SetParameters(Parameters);
 	}
+}
 
-	// Get the color data from the rendered component
-	FLinearColor PixelColor;
-	if (RendererComponent->GetDownsampleBufferPixel(DownsamplePixelIndex, PixelColor))
+UDMXPixelMappingRendererComponent* UDMXPixelMappingMatrixCellComponent::GetRendererComponent() const
+{
+	for (UDMXPixelMappingBaseComponent* ParentComponent = GetParent(); ParentComponent; ParentComponent = ParentComponent->GetParent())
 	{
-		ColorSpace->SetRGBA(PixelColor);
-		AttributeToNormalizedValueMap = ColorSpace->GetAttributeNameToValueMap();
+		if (UDMXPixelMappingRendererComponent* RendererComponent = Cast<UDMXPixelMappingRendererComponent>(ParentComponent))
+		{
+			return RendererComponent;
+		}
 	}
+	return nullptr;
+}
 
-	return AttributeToNormalizedValueMap;
+FLinearColor UDMXPixelMappingMatrixCellComponent::GetPixelMapColor() const
+{
+	return PixelMapRenderElement->GetColor();
 }
 
 #undef LOCTEXT_NAMESPACE
diff --git a/Engine/Plugins/VirtualProduction/DMX/DMXPixelMapping/Source/DMXPixelMappingRuntime/Private/Components/DMXPixelMappingMatrixComponent.cpp b/Engine/Plugins/VirtualProduction/DMX/DMXPixelMapping/Source/DMXPixelMappingRuntime/Private/Components/DMXPixelMappingMatrixComponent.cpp
--- a/Engine/Plugins/VirtualProduction/DMX/DMXPixelMapping/Source/DMXPixelMappingRuntime/Private/Components/DMXPixelMappingMatrixComponent.cpp
+++ b/Engine/Plugins/VirtualProduction/DMX/DMXPixelMapping/Source/DMXPixelMappingRuntime/Private/Components/DMXPixelMappingMatrixComponent.cpp
@@ -2,6 +2,11 @@
 
 #include "Components/DMXPixelMappingMatrixComponent.h"
 
+#include "Algo/Sort.h"
+#include "ColorSpace/DMXPixelMappingColorSpace_RGBCMY.h"
+#include "Components/DMXPixelMappingFixtureGroupComponent.h"
+#include "Components/DMXPixelMappingMatrixCellComponent.h"
+#include "DMXConversions.h"
 #include "DMXPixelMapping.h"
 #include "DMXPixelMappingMainStreamObjectVersion.h"
 #include "DMXPixelMappingRuntimeCommon.h"
@@ -10,6 +15,7 @@
 #include "ColorSpace/DMXPixelMappingColorSpace_RGBCMY.h"
 #include "Components/DMXPixelMappingFixtureGroupComponent.h"
 #include "Components/DMXPixelMappingMatrixCellComponent.h"
+#include "Components/DMXPixelMappingRendererComponent.h"
 #include "IO/DMXOutputPort.h"
 #include "Library/DMXEntityFixturePatch.h"
 #include "Library/DMXLibrary.h"
@@ -23,6 +29,7 @@
 
 #define LOCTEXT_NAMESPACE "DMXPixelMappingMatrixComponent"
 
+
 UDMXPixelMappingMatrixComponent::UDMXPixelMappingMatrixComponent()
 {
 	ColorSpaceClass = UDMXPixelMappingColorSpace_RGBCMY::StaticClass();
@@ -285,23 +292,46 @@ void UDMXPixelMappingMatrixComponent::ResetDMX()
 
 void UDMXPixelMappingMatrixComponent::SendDMX()
 {
-	UDMXEntityFixturePatch* FixturePatch = FixturePatchRef.GetFixturePatch();
+	if (!ensureMsgf(ColorSpace, TEXT("Invalid color space in Pixel Mapping Matrix Component '%s'"), *GetName()))
+	{
+		return;
+	}
 
-	if (FixturePatch)
+	UDMXEntityFixturePatch* FixturePatch = FixturePatchRef.GetFixturePatch();
+	const FDMXFixtureMode* ActiveModePtr = FixturePatch ? FixturePatch->GetActiveMode() : nullptr;
+	if (FixturePatch && ActiveModePtr && ActiveModePtr->bFixtureMatrixEnabled)
 	{
-		// An array of attribute to value maps for each child, in order of the childs
-		TArray<FDMXNormalizedAttributeValueMap> AttributeToValueMapArray;
+		TArray<UDMXPixelMappingMatrixCellComponent*> CellComponents;
+		Algo::TransformIf(Children, CellComponents,
+			[](UDMXPixelMappingBaseComponent* BaseComponent)
+			{
+				return BaseComponent != nullptr;
+			},
+			[](UDMXPixelMappingBaseComponent* BaseComponent)
+			{
+				return Cast<UDMXPixelMappingMatrixCellComponent>(BaseComponent);
+			});
+		Algo::SortBy(CellComponents, [](UDMXPixelMappingMatrixCellComponent* Component)
+			{
+				// Sort by cell coordinates
+				return Component->GetCellCoordinate().Y * Component->GetCellCoordinate().X + Component->GetCellCoordinate().X;
+			});
 
 		// Accumulate matrix cell data
-		for (UDMXPixelMappingBaseComponent* Component : Children)
+		TArray<FDMXNormalizedAttributeValueMap> AttributeToValueMapArray;
+		AttributeToValueMapArray.Reserve(CellComponents.Num());
+		Algo::Transform(CellComponents, AttributeToValueMapArray, [this](UDMXPixelMappingMatrixCellComponent* CellComponent)
+			{	
+				// Get the color data from the rendered component
+				ColorSpace->SetRGBA(CellComponent->GetPixelMapColor());
+
+				FDMXNormalizedAttributeValueMap MapStruct;
+				MapStruct.Map = ColorSpace->GetAttributeNameToValueMap();
+				return MapStruct;
+			});
+		if (!ensureMsgf(AttributeToValueMapArray.Num() == CellComponents.Num(), TEXT("Mismatch num cell attributes and num attribute to value maps. Cannot send DMX for Matrix Component %s"), *GetUserFriendlyName()))
 		{
-			if (UDMXPixelMappingMatrixCellComponent* CellComponent = Cast<UDMXPixelMappingMatrixCellComponent>(Component))
-			{
-				FDMXNormalizedAttributeValueMap NormalizedAttributeToValueMap;
-				NormalizedAttributeToValueMap.Map = CellComponent->CreateAttributeValues();
-
-				AttributeToValueMapArray.Add(NormalizedAttributeToValueMap);
-			}
+			return;
 		}
 
 		// Apply matrix modulators
@@ -310,39 +340,43 @@ void UDMXPixelMappingMatrixComponent::SendDMX()
 			Modulator->ModulateMatrix(FixturePatch, AttributeToValueMapArray, AttributeToValueMapArray);
 		}
 
-		for (int32 IndexChild = 0; IndexChild < Children.Num(); IndexChild++)
+		// Write channel values
+		const int32 MatrixStartingChannel = FixturePatch->GetStartingChannel() + ActiveModePtr->FixtureMatrixConfig.FirstCellChannel - 1;
+		const int32 MatrixCellSize = ActiveModePtr->FixtureMatrixConfig.GetNumChannels();
+
+		int32 ChannelOffset = 0;
+		TMap<int32, uint8> ChannelToValueMap;
+		for (int32 CellIndex = 0; CellIndex < CellComponents.Num(); CellIndex++)
 		{
-			if (UDMXPixelMappingMatrixCellComponent* CellComponent = Cast<UDMXPixelMappingMatrixCellComponent>(Children[IndexChild]))
+			for (const FDMXFixtureCellAttribute& CellAttribute : ActiveModePtr->FixtureMatrixConfig.CellAttributes)
 			{
-				// Relies on the order of childs and AttributeToValueMapArray didn't change during the lifetime of this function!
-				if (!AttributeToValueMapArray.IsValidIndex(IndexChild))
+				float* ValuePtr = AttributeToValueMapArray[CellIndex].Map.Find(CellAttribute.Attribute);
+				if (ValuePtr)
 				{
-					break;
+					const TArray<uint8> Values = FDMXConversions::NormalizedDMXValueToByteArray(*ValuePtr, CellAttribute.DataType, CellAttribute.bUseLSBMode);
+					for (uint8 Value : Values)
+					{
+						const int32 Channel = MatrixStartingChannel + ChannelOffset;
+						ChannelToValueMap.Add(Channel, Value);
+					}
 				}
 
-				TMap<int32, uint8> ChannelToValueMap;
-				for (const TTuple<FDMXAttributeName, float>& AttributeValuePair : AttributeToValueMapArray[IndexChild].Map)
-				{
-					FDMXPixelMappingRuntimeUtils::ConvertNormalizedAttributeValueToChannelValue(FixturePatch, AttributeValuePair.Key, AttributeValuePair.Value, ChannelToValueMap);
-
-					FixturePatch->SendNormalizedMatrixCellValue(CellComponent->CellCoordinate, AttributeValuePair.Key, AttributeValuePair.Value);
-				}
+				ChannelOffset += CellAttribute.GetNumChannels();
 			}
 		}
 
-		// Send normal modulators. This is important to allow modulators to generate attribute values
+		// Write modulator generated non-matrix attributes. This is important to allow modulators to generate non-matrix attribute values
 		TMap<FDMXAttributeName, float> ModulatorGeneratedAttributeValueMap;
 		for (UDMXModulator* Modulator : Modulators)
 		{
 			Modulator->Modulate(FixturePatch, ModulatorGeneratedAttributeValueMap, ModulatorGeneratedAttributeValueMap);
 		}
-
-		TMap<int32, uint8> ChannelToValueMap;
 		for (const TTuple<FDMXAttributeName, float>& AttributeValuePair : ModulatorGeneratedAttributeValueMap)
 		{
 			FDMXPixelMappingRuntimeUtils::ConvertNormalizedAttributeValueToChannelValue(FixturePatch, AttributeValuePair.Key, AttributeValuePair.Value, ChannelToValueMap);
 		}
 
+		// Send DMX
 		if (UDMXLibrary* Library = FixturePatch->GetParentLibrary())
 		{
 			for (const FDMXOutputPortSharedRef& OutputPort : Library->GetOutputPorts())
@@ -353,16 +387,6 @@ void UDMXPixelMappingMatrixComponent::SendDMX()
 	}
 }
 
-void UDMXPixelMappingMatrixComponent::QueueDownsample()
-{
-	ForEachChild([&](UDMXPixelMappingBaseComponent* InComponent) {
-		if (UDMXPixelMappingOutputComponent* Component = Cast<UDMXPixelMappingOutputComponent>(InComponent))
-		{
-			Component->QueueDownsample();
-		}
-	}, false);
-}
-
 bool UDMXPixelMappingMatrixComponent::CanBeMovedTo(const UDMXPixelMappingBaseComponent* Component) const
 {
 	if (const UDMXPixelMappingFixtureGroupComponent* FixtureGroupComponent = Cast<UDMXPixelMappingFixtureGroupComponent>(Component))
@@ -435,30 +459,21 @@ void UDMXPixelMappingMatrixComponent::HandleSizeChanged()
 		CellSize = FVector2D(GetSize().X / CoordinateGrid.X, GetSize().Y / CoordinateGrid.Y);
 	}
 
-	if (Children.Num() > 0)
-	{
-		constexpr bool bUpdateSizeRecursive = false;
-		ForEachChildOfClass<UDMXPixelMappingMatrixCellComponent>([this](UDMXPixelMappingMatrixCellComponent* ChildComponent)
-			{
+	constexpr bool bUpdateSizeRecursive = false;
+	ForEachChildOfClass<UDMXPixelMappingMatrixCellComponent>([this](UDMXPixelMappingMatrixCellComponent* ChildComponent)
+		{
 #if WITH_EDITOR
-				if (ChildComponent->IsLockInDesigner())
-				{
-					ChildComponent->SetSize(CellSize);
-					ChildComponent->SetPosition(GetPosition() + FVector2D(CellSize * ChildComponent->GetCellCoordinate()));
-				}
-#else
+			if (ChildComponent->IsLockInDesigner())
+			{
 				ChildComponent->SetSize(CellSize);
 				ChildComponent->SetPosition(GetPosition() + FVector2D(CellSize * ChildComponent->GetCellCoordinate()));
+			}
+#else
+			ChildComponent->SetSize(CellSize);
+			ChildComponent->SetPosition(GetPosition() + FVector2D(CellSize * ChildComponent->GetCellCoordinate()));
 #endif // WITH_EDITOR
 
-			}, bUpdateSizeRecursive);
-	}
-	else
-	{
-		// Set the default size
-		const FVector2D DefaultSize(32.f, 32.f);
-		SetSize(DefaultSize);
-	}
+		}, bUpdateSizeRecursive);
 
 #if WITH_EDITOR
 PRAGMA_DISABLE_DEPRECATION_WARNINGS
@@ -519,6 +534,19 @@ void UDMXPixelMappingMatrixComponent::HandleMatrixChanged()
 	GetOnMatrixChanged().Broadcast(PixelMapping, this);
 }
 
+void UDMXPixelMappingMatrixComponent::QueueDownsample()
+{
+	// Deprecated 5.3
+	PRAGMA_DISABLE_DEPRECATION_WARNINGS
+	ForEachChild([&](UDMXPixelMappingBaseComponent* InComponent) {
+		if (UDMXPixelMappingOutputComponent* Component = Cast<UDMXPixelMappingOutputComponent>(InComponent))
+		{
+			Component->QueueDownsample();
+		}
+	}, false);
+	PRAGMA_ENABLE_DEPRECATION_WARNINGS
+}
+
 void UDMXPixelMappingMatrixComponent::OnFixtureTypeChanged(const UDMXEntityFixtureType* FixtureType)
 {
 	if (UDMXEntityFixturePatch* FixturePatch = FixturePatchRef.GetFixturePatch())
diff --git a/Engine/Plugins/VirtualProduction/DMX/DMXPixelMapping/Source/DMXPixelMappingRuntime/Private/Components/DMXPixelMappingOutputComponent.cpp b/Engine/Plugins/VirtualProduction/DMX/DMXPixelMapping/Source/DMXPixelMappingRuntime/Private/Components/DMXPixelMappingOutputComponent.cpp
--- a/Engine/Plugins/VirtualProduction/DMX/DMXPixelMapping/Source/DMXPixelMappingRuntime/Private/Components/DMXPixelMappingOutputComponent.cpp
+++ b/Engine/Plugins/VirtualProduction/DMX/DMXPixelMapping/Source/DMXPixelMappingRuntime/Private/Components/DMXPixelMappingOutputComponent.cpp
@@ -273,20 +273,35 @@ bool UDMXPixelMappingOutputComponent::OverlapsComponent(UDMXPixelMappingOutputCo
 
 void UDMXPixelMappingOutputComponent::SetPosition(const FVector2D& Position) 
 {
-	if (Position.X != PositionX || Position.Y != PositionY)
-	{
-		PositionX = Position.X;
-		PositionY = Position.Y;
-	}
+	PositionX = Position.X;
+	PositionY = Position.Y;
+}
+
+FVector2D UDMXPixelMappingOutputComponent::GetPosition() const
+{
+	return FVector2D(PositionX, PositionY);
 }
 
 void UDMXPixelMappingOutputComponent::SetSize(const FVector2D& Size) 
 {
-	SizeX = Size.X;
-	SizeY = Size.Y;
+	SizeX = FMath::Max(Size.X, 1.f);
+	SizeY = FMath::Max(Size.Y, 1.f);
+
+	// Limit all components size to max texture size, logged
+	const uint32 MaxTextureDimensions = GetMax2DTextureDimension();
+	if (SizeX > MaxTextureDimensions || SizeY > MaxTextureDimensions)
+	{
+		SizeX = MaxTextureDimensions;
+		SizeY = MaxTextureDimensions;
+	}
+}
 
-	SizeX = FMath::Max(SizeX, 1.f);
-	SizeY = FMath::Max(SizeY, 1.f);
+void UDMXPixelMappingOutputComponent::InvalidatePixelMapRenderer()
+{
+	if (UDMXPixelMappingRendererComponent* RendererComponent = GetRendererComponent())
+	{
+		RendererComponent->InvalidatePixelMapRenderer();
+	}
 }
 
 UDMXPixelMappingRendererComponent* UDMXPixelMappingOutputComponent::FindRendererComponent() const
diff --git a/Engine/Plugins/VirtualProduction/DMX/DMXPixelMapping/Source/DMXPixelMappingRuntime/Private/Components/DMXPixelMappingRendererComponent.cpp b/Engine/Plugins/VirtualProduction/DMX/DMXPixelMapping/Source/DMXPixelMappingRuntime/Private/Components/DMXPixelMappingRendererComponent.cpp
--- a/Engine/Plugins/VirtualProduction/DMX/DMXPixelMapping/Source/DMXPixelMappingRuntime/Private/Components/DMXPixelMappingRendererComponent.cpp
+++ b/Engine/Plugins/VirtualProduction/DMX/DMXPixelMapping/Source/DMXPixelMappingRuntime/Private/Components/DMXPixelMappingRendererComponent.cpp
@@ -9,6 +9,7 @@
 #include "Components/DMXPixelMappingMatrixComponent.h"
 #include "Components/DMXPixelMappingRootComponent.h"
 #include "Components/DMXPixelMappingScreenComponent.h"
+#include "DMXPixelMappingPixelMapRenderer.h"
 #include "DMXPixelMappingPreprocessRenderer.h"
 #include "DMXPixelMappingMainStreamObjectVersion.h"
 #include "DMXPixelMappingTypes.h"
@@ -27,15 +28,10 @@
 #include "LevelEditor.h"
 #endif
 
-
 DECLARE_CYCLE_STAT(TEXT("PixelMapping Render"), STAT_DMXPixelMappingRender, STATGROUP_DMX);
 DECLARE_CYCLE_STAT(TEXT("PixelMapping SendDMX"), STAT_DMXPixelMappingSendDMX, STATGROUP_DMX);
-DECLARE_CYCLE_STAT(TEXT("PixelMapping RenderInputTexture"), STAT_RenderInputTexture, STATGROUP_DMX);
-DECLARE_CYCLE_STAT(TEXT("PixelMapping GetTotalDownsamplePixelCount"), STAT_GetTotalDownsamplePixelCount, STATGROUP_DMX);
-
+DECLARE_CYCLE_STAT(TEXT("PixelMapping RenderInputTexture"), STAT_DMXPixelMappingRenderInputTexture, STATGROUP_DMX);
 
-const FIntPoint UDMXPixelMappingRendererComponent::MaxDownsampleBufferTargetSize = FIntPoint(4096);
-const FLinearColor UDMXPixelMappingRendererComponent::ClearTextureColor = FLinearColor::Black;
 
 UDMXPixelMappingRendererComponent::UDMXPixelMappingRendererComponent()
 {
@@ -56,7 +52,12 @@ UDMXPixelMappingRendererComponent::UDMXPixelMappingRendererComponent()
 #endif
 	
 	PreprocessRenderer = CreateDefaultSubobject<UDMXPixelMappingPreprocessRenderer>("PreprocessRenderer");
+	PixelMapRenderer = CreateDefaultSubobject<UDMXPixelMappingPixelMapRenderer>("PixelMapRenderer");
+
 	Brightness = 1.0f;
+
+	UDMXPixelMappingBaseComponent::GetOnComponentAdded().AddUObject(this, &UDMXPixelMappingRendererComponent::OnComponentAddedOrRemoved);
+	UDMXPixelMappingBaseComponent::GetOnComponentRemoved().AddUObject(this, &UDMXPixelMappingRendererComponent::OnComponentAddedOrRemoved);
 }
 
 const FName& UDMXPixelMappingRendererComponent::GetNamePrefix()
@@ -67,7 +68,7 @@ const FName& UDMXPixelMappingRendererComponent::GetNamePrefix()
 
 bool UDMXPixelMappingRendererComponent::CanBeMovedTo(const UDMXPixelMappingBaseComponent* Component) const
 {
-	return Component && Component->IsA<UDMXPixelMappingRootComponent>();
+	return Component && Component->GetClass() == UDMXPixelMappingRootComponent::StaticClass();
 }
 
 void UDMXPixelMappingRendererComponent::Serialize(FArchive& Ar)
@@ -101,26 +102,41 @@ void UDMXPixelMappingRendererComponent::PostInitProperties()
 
 	if (!IsTemplate())
 	{
-		Initialize();
+		InvalidatePreprocessRenderer();
 	}
 }
 
 void UDMXPixelMappingRendererComponent::PostLoad()
 {
 	Super::PostLoad();
-	
+
 	if (!IsTemplate())
 	{
-		Initialize();
+		InvalidatePreprocessRenderer();
 	}
 }
 
+#if WITH_EDITOR
+void UDMXPixelMappingRendererComponent::PostEditUndo()
+{
+	Super::PostEditUndo();
+
+	InvalidatePixelMapRenderer();
+}
+#endif // WITH_EDITOR
+
 #if WITH_EDITOR
 void UDMXPixelMappingRendererComponent::PostEditChangeChainProperty(FPropertyChangedChainEvent& PropertyChangedChainEvent)
 {
 	// Call the parent at the first place
 	Super::PostEditChangeChainProperty(PropertyChangedChainEvent);
 
+	if (PropertyChangedChainEvent.ChangeType != EPropertyChangeType::Interactive)
+	{
+		InvalidatePreprocessRenderer();
+	}
+
+	PRAGMA_DISABLE_DEPRECATION_WARNINGS
 	const FName PropertyName = PropertyChangedChainEvent.GetPropertyName();
 	if (PropertyName == GET_MEMBER_NAME_CHECKED(UDMXPixelMappingRendererComponent, Brightness))
 	{
@@ -130,85 +146,48 @@ void UDMXPixelMappingRendererComponent::PostEditChangeChainProperty(FPropertyCha
 			Renderer->SetBrightness(Brightness);
 		}
 	}
-
-	if (PropertyChangedChainEvent.ChangeType == EPropertyChangeType::Interactive)
-	{
-		// Apply interactive changes only once per frame
-		static uint64 Frame = GFrameCounter;
-		if (Frame != GFrameCounter)
-		{
-			Initialize();
-			Frame = GFrameCounter;
-		}
-	}
-	else
-	{
-		Initialize();
-	}
+	PRAGMA_ENABLE_DEPRECATION_WARNINGS
 }
 #endif // WITH_EDITOR
 
-#if WITH_EDITOR
-void UDMXPixelMappingRendererComponent::RenderEditorPreviewTexture()
+void UDMXPixelMappingRendererComponent::InvalidatePreprocessRenderer()
 {
-	if (!DownsampleBufferTarget)
-	{
-		return;
+	PRAGMA_DISABLE_DEPRECATION_WARNINGS
+	if (!PixelMappingRenderer_DEPRECATED.IsValid())
+	{	
+		// To keep support for deprecated functions, still create the old pixel mapping renderer 
+		PixelMappingRenderer_DEPRECATED = IDMXPixelMappingRendererModule::Get().CreateRenderer();
+		PixelMappingRenderer_DEPRECATED->SetBrightness(Brightness);
 	}
+	PRAGMA_ENABLE_DEPRECATION_WARNINGS
 
-	const TSharedPtr<IDMXPixelMappingRenderer>& Renderer = GetRenderer();
-	if (!ensure(Renderer))
+	UserWidget = nullptr;
+	switch (RendererType)
 	{
-		return;
-	}
-
-	TArray<FDMXPixelMappingDownsamplePixelPreviewParam> PixelPreviewParams;
-	PixelPreviewParams.Reserve(DownsamplePixelCount);
-	
-	ForEachChild([this, &PixelPreviewParams](UDMXPixelMappingBaseComponent* InComponent) {
-		if(UDMXPixelMappingScreenComponent* ScreenComponent = Cast<UDMXPixelMappingScreenComponent>(InComponent))
-		{
-			const FVector2D SizePixel = ScreenComponent->GetScreenPixelSize();
-			const int32 DownsampleIndexStart = ScreenComponent->GetPixelDownsamplePositionRange().Key;
-			const int32 PositionX = ScreenComponent->GetPosition().X;
-			const int32 PositionY = ScreenComponent->GetPosition().Y;
+	case(EDMXPixelMappingRendererType::Texture):
+		PreprocessRenderer->SetInputTexture(InputTexture.Get());
+		break;
 
-			ScreenComponent->ForEachPixel([this, &PixelPreviewParams, SizePixel, PositionX, PositionY, DownsampleIndexStart](const int32 InXYIndex, const int32 XIndex, const int32 YIndex)
-				{
-					FDMXPixelMappingDownsamplePixelPreviewParam PixelPreviewParam;
-					PixelPreviewParam.ScreenPixelSize = SizePixel;
-					PixelPreviewParam.ScreenPixelPosition = FVector2D(PositionX + SizePixel.X * XIndex, PositionY + SizePixel.Y * YIndex);
-					PixelPreviewParam.DownsamplePosition = GetPixelPosition(InXYIndex + DownsampleIndexStart);
+	case(EDMXPixelMappingRendererType::Material):
+		PreprocessRenderer->SetInputMaterial(InputMaterial.Get());
+		break;
 
-					PixelPreviewParams.Add(MoveTemp(PixelPreviewParam));
-				});
-		}
-		else if (UDMXPixelMappingOutputDMXComponent* Component = Cast<UDMXPixelMappingOutputDMXComponent>(InComponent))
-		{
-			FDMXPixelMappingDownsamplePixelPreviewParam PixelPreviewParam;
-			PixelPreviewParam.ScreenPixelSize = Component->GetSize();
-			PixelPreviewParam.ScreenPixelPosition = Component->GetPosition();
-			PixelPreviewParam.DownsamplePosition = GetPixelPosition(Component->GetDownsamplePixelIndex());
+	case(EDMXPixelMappingRendererType::UMG):
+		UserWidget = CreateWidget(TryGetWorld(), InputWidget);
+		PreprocessRenderer->SetInputUserWidget(UserWidget.Get());
+		break;
 
-			PixelPreviewParams.Add(MoveTemp(PixelPreviewParam));
-		}
-	}, true);
+	default:
+		checkf(0, TEXT("Invalid Renderer Type in DMXPixelMappingRendererComponent"));
+	}
 
-	Renderer->RenderPreview(GetPreviewRenderTarget()->GetResource(), DownsampleBufferTarget->GetResource(), MoveTemp(PixelPreviewParams));
+	SetSize(PreprocessRenderer->GetResultingSize2D());
 }
-#endif // WITH_EDITOR
 
-#if WITH_EDITOR
-UTextureRenderTarget2D* UDMXPixelMappingRendererComponent::GetPreviewRenderTarget()
+void UDMXPixelMappingRendererComponent::InvalidatePixelMapRenderer()
 {
-	if (PreviewRenderTarget == nullptr)
-	{
-		PreviewRenderTarget = CreateRenderTarget(TEXT("DMXPreviewRenderTarget"));
-	}
-
-	return PreviewRenderTarget;
+	bInvalidatePixelMap = true;
 }
-#endif // WITH_EDITOR
 
 bool UDMXPixelMappingRendererComponent::GetPixelMappingComponentModulators(FDMXEntityFixturePatchRef FixturePatchRef, TArray<UDMXModulator*>& DMXModulators)
 {
@@ -241,42 +220,10 @@ bool UDMXPixelMappingRendererComponent::GetPixelMappingComponentModulators(FDMXE
 	return false;
 }
 
-#if WITH_EDITOR
-TSharedRef<SWidget> UDMXPixelMappingRendererComponent::TakeWidget()
-{
-	if (!ComponentsCanvas.IsValid())
-	{
-		ComponentsCanvas =
-			SNew(SConstraintCanvas);
-	}
-
-	PRAGMA_DISABLE_DEPRECATION_WARNINGS
-	ForEachChild([&](UDMXPixelMappingBaseComponent* InComponent) {
-		if (UDMXPixelMappingOutputComponent* Component = Cast<UDMXPixelMappingOutputComponent>(InComponent))
-		{
-			// Build all child DMX pixel mapping slots
-			Component->BuildSlot(ComponentsCanvas.ToSharedRef());
-		}
-	}, true);
-	PRAGMA_ENABLE_DEPRECATION_WARNINGS
-
-	return ComponentsCanvas.ToSharedRef();
-}
-#endif // WITH_EDITOR
-
-#if WITH_EDITOR
-void UDMXPixelMappingRendererComponent::ResizePreviewRenderTarget(uint32 InSizeX, uint32 InSizeY)
+TArray<TSharedRef<UE::DMXPixelMapping::Rendering::FPixelMapRenderElement>> UDMXPixelMappingRendererComponent::GetPixelMapRenderElements() const
 {
-	UTextureRenderTarget2D* Target = GetPreviewRenderTarget();
-
-	if ((InSizeX > 0 && InSizeY > 0) && (Target->SizeX != InSizeX || Target->SizeY != InSizeY))
-	{
-		check(Target);
-		Target->ResizeTarget(InSizeX, InSizeY);
-		Target->UpdateResourceImmediate();
-	}
+	return PixelMapRenderElements;
 }
-#endif // WITH_EDITOR
 
 void UDMXPixelMappingRendererComponent::ResetDMX()
 {
@@ -306,48 +253,38 @@ void UDMXPixelMappingRendererComponent::Render()
 {
 	SCOPE_CYCLE_COUNTER(STAT_DMXPixelMappingRender);
 
-	if (!PixelMappingRenderer.IsValid())
-	{
-		return;
-	}
+	PreprocessRenderer->Render();
 
-	// 1. Render the input texture
-	RendererInputTexture();
-	UTexture* DownsampleInputTexture = GetRenderedInputTexture();
-	if (!DownsampleInputTexture)
+	UTexture* RenderedInputTexture = PreprocessRenderer->GetRenderedTexture();
+	if (!RenderedInputTexture)
 	{
 		return;
 	}
 
-	// 2. Make sure there is the DownsampleBufferTarget exists and can size can hold all pixels
-	CreateOrUpdateDownsampleBufferTarget();
+	// Update render elements if invalidated
+	if (bInvalidatePixelMap)
+	{
+		PixelMapRenderElements.Reset();
 
-	// 3. reserve enough space for pixels params
-	DownsamplePixelParams.Reset(DownsamplePixelCount);
+		constexpr bool bRecursive = true;
+		ForEachChild([this](UDMXPixelMappingBaseComponent* Component)
+			{
+				if (UDMXPixelMappingFixtureGroupItemComponent* FixtureGroupItemComponent = Cast< UDMXPixelMappingFixtureGroupItemComponent>(Component))
+				{
+					PixelMapRenderElements.Add(FixtureGroupItemComponent->GetOrCreatePixelMapRenderElement());
+				}
+				else if (UDMXPixelMappingMatrixCellComponent* MatrixCellComponent = Cast<UDMXPixelMappingMatrixCellComponent>(Component))
+				{
+					PixelMapRenderElements.Add(MatrixCellComponent->GetOrCreatePixelMapRenderElement());
+				}
+			}, bRecursive);
 
-	// 4. Loop through all child pixels in order to get pixels downsample params for rendering
-	ForEachChild([&](UDMXPixelMappingBaseComponent* InComponent) {
-		if (UDMXPixelMappingOutputComponent* Component = Cast<UDMXPixelMappingOutputComponent>(InComponent))
-		{
-			Component->QueueDownsample();
-		}
-	}, false);
+		PixelMapRenderer->SetElements(PixelMapRenderElements);
 
-	// 5. Make sure pixel count the same with pixel params set number
-	if (!ensure(DownsamplePixelParams.Num() == DownsamplePixelCount))
-	{
-		DownsamplePixelParams.Empty();
-		return;
+		bInvalidatePixelMap = false;
 	}
 
-	// 6. Downsample all pixels
-	PixelMappingRenderer->DownsampleRender(
-		DownsampleInputTexture->GetResource(),
-		DownsampleBufferTarget->GetResource(),
-		DownsampleBufferTarget->GameThread_GetRenderTargetResource(),
-		DownsamplePixelParams, // Copy Set to GPU thread, no empty function call needed
-		[this](TArray<FLinearColor>&& InDownsampleBuffer, FIntRect InRect) { SetDownsampleBuffer(MoveTemp(InDownsampleBuffer), InRect); }
-	);
+	PixelMapRenderer->Render(RenderedInputTexture, Brightness);
 }
 
 void UDMXPixelMappingRendererComponent::RenderAndSendDMX()
@@ -356,16 +293,45 @@ void UDMXPixelMappingRendererComponent::RenderAndSendDMX()
 	SendDMX();
 }
 
-FIntPoint UDMXPixelMappingRendererComponent::GetPixelPosition(int32 InPosition) const
+UTexture* UDMXPixelMappingRendererComponent::GetRenderedInputTexture() const
 {
-	const int32 YRows = InPosition / MaxDownsampleBufferTargetSize.X;
-	return FIntPoint(InPosition % MaxDownsampleBufferTargetSize.X, YRows);
+	return  PreprocessRenderer ? PreprocessRenderer->GetRenderedTexture() : nullptr;
 }
 
-UWorld* UDMXPixelMappingRendererComponent::GetWorld() const
+void UDMXPixelMappingRendererComponent::OnComponentAddedOrRemoved(UDMXPixelMapping* PixelMapping, UDMXPixelMappingBaseComponent* Component)
+{
+	InvalidatePixelMapRenderer();
+}
+
+UWorld* UDMXPixelMappingRendererComponent::TryGetWorld() const
 {
 	UWorld* World = nullptr;
+	if (GIsEditor)
+	{
+#if WITH_EDITOR
+		World = GEditor->GetEditorWorldContext().World();
+#endif
+	}
+	else
+	{
+		World = GWorld;
+	}
+
+	return World;
+}
+
+
+///////////////////////////////
+// DEPRECATED MEMBERS 5.3
+
+const FIntPoint UDMXPixelMappingRendererComponent::MaxDownsampleBufferTargetSize_DEPRECATED = FIntPoint(4096);
+const FLinearColor UDMXPixelMappingRendererComponent::ClearTextureColor_DEPRECATED = FLinearColor::Black;
+
+UWorld* UDMXPixelMappingRendererComponent::GetWorld() const
+{
+	// DEPRECATED 5.3
 
+	UWorld* World = nullptr;
 	if (GIsEditor)
 	{
 #if WITH_EDITOR
@@ -380,88 +346,118 @@ UWorld* UDMXPixelMappingRendererComponent::GetWorld() const
 	return World;
 }
 
-void UDMXPixelMappingRendererComponent::Initialize()
+#if WITH_EDITOR
+void UDMXPixelMappingRendererComponent::RenderEditorPreviewTexture()
 {
-	if (!PixelMappingRenderer.IsValid())
+	// DEPRECATED 5.3
+
+	if (!DownsampleBufferTarget_DEPRECATED)
 	{
-		PixelMappingRenderer = IDMXPixelMappingRendererModule::Get().CreateRenderer();
-		PixelMappingRenderer->SetBrightness(Brightness);
+		return;
 	}
 
-	switch (RendererType)
+	PRAGMA_DISABLE_DEPRECATION_WARNINGS
+	const TSharedPtr<IDMXPixelMappingRenderer>& Renderer = GetRenderer();
+	if (!ensure(Renderer))
 	{
-	case(EDMXPixelMappingRendererType::Texture):
-		PreprocessRenderer->SetInputTexture(InputTexture.Get());
-		break;
+		return;
+	}
 
-	case(EDMXPixelMappingRendererType::Material):
-		PreprocessRenderer->SetInputMaterial(InputMaterial.Get());
-		break;
+	TArray<FDMXPixelMappingDownsamplePixelPreviewParam> PixelPreviewParams;
+	PixelPreviewParams.Reserve(DownsamplePixelCount_DEPRECATED);
+	
+	ForEachChild([this, &PixelPreviewParams](UDMXPixelMappingBaseComponent* InComponent) {
+		if(UDMXPixelMappingScreenComponent* ScreenComponent = Cast<UDMXPixelMappingScreenComponent>(InComponent))
+		{
+			const FVector2D SizePixel = ScreenComponent->GetScreenPixelSize();
+			const int32 DownsampleIndexStart = ScreenComponent->GetPixelDownsamplePositionRange().Key;
+			const int32 PositionX = ScreenComponent->GetPosition().X;
+			const int32 PositionY = ScreenComponent->GetPosition().Y;
 
-	case(EDMXPixelMappingRendererType::UMG):
-		UserWidget = CreateWidget(GetWorld(), InputWidget);
-		PreprocessRenderer->SetInputUserWidget(UserWidget.Get());
-		break;
+			ScreenComponent->ForEachPixel([this, &PixelPreviewParams, SizePixel, PositionX, PositionY, DownsampleIndexStart](const int32 InXYIndex, const int32 XIndex, const int32 YIndex)
+				{
+					FDMXPixelMappingDownsamplePixelPreviewParam PixelPreviewParam;
+					PixelPreviewParam.ScreenPixelSize = SizePixel;
+					PixelPreviewParam.ScreenPixelPosition = FVector2D(PositionX + SizePixel.X * XIndex, PositionY + SizePixel.Y * YIndex);
+					PixelPreviewParam.DownsamplePosition = GetPixelPosition(InXYIndex + DownsampleIndexStart);
 
-	default:
-		checkf(0, TEXT("Invalid Renderer Type in DMXPixelMappingRendererComponent"));
-	}
+					PixelPreviewParams.Add(MoveTemp(PixelPreviewParam));
+				});
+		}
+		else if (UDMXPixelMappingOutputDMXComponent* Component = Cast<UDMXPixelMappingOutputDMXComponent>(InComponent))
+		{
+			FDMXPixelMappingDownsamplePixelPreviewParam PixelPreviewParam;
+			PixelPreviewParam.ScreenPixelSize = Component->GetSize();
+			PixelPreviewParam.ScreenPixelPosition = Component->GetPosition();
+			PixelPreviewParam.DownsamplePosition = GetPixelPosition(Component->GetDownsamplePixelIndex());
 
-	SetSize(PreprocessRenderer->GetResultingSize2D());
+			PixelPreviewParams.Add(MoveTemp(PixelPreviewParam));
+		}
+	}, true);
+
+	Renderer->RenderPreview(GetPreviewRenderTarget()->GetResource(), DownsampleBufferTarget_DEPRECATED->GetResource(), MoveTemp(PixelPreviewParams));
+	PRAGMA_ENABLE_DEPRECATION_WARNINGS
 }
+#endif // WITH_EDITOR
 
-UTextureRenderTarget2D* UDMXPixelMappingRendererComponent::CreateRenderTarget(const FName& InBaseName)
-{
-	const FName TargetName = MakeUniqueObjectName(this, UTextureRenderTarget2D::StaticClass(), InBaseName);
-	UTextureRenderTarget2D* RenderTarget = NewObject<UTextureRenderTarget2D>(this, TargetName);
-	RenderTarget->ClearColor = ClearTextureColor;
-	constexpr bool bInForceLinearGamma = false;
-	RenderTarget->InitCustomFormat(GetSize().X, GetSize().Y, EPixelFormat::PF_B8G8R8A8, bInForceLinearGamma);
+#if WITH_EDITOR
+UTextureRenderTarget2D* UDMXPixelMappingRendererComponent::GetPreviewRenderTarget()
+{	
+	// DEPRECATED 5.3
+	PRAGMA_DISABLE_DEPRECATION_WARNINGS
 
-	return RenderTarget;
-}
+	if (PreviewRenderTarget_DEPRECATED == nullptr)
+	{
+		PreviewRenderTarget_DEPRECATED = CreateRenderTarget(TEXT("DMXPreviewRenderTarget_DEPRECATED"));
+	}
 
-void UDMXPixelMappingRendererComponent::RendererInputTexture()
-{
-	SCOPE_CYCLE_COUNTER(STAT_RenderInputTexture);
+	PRAGMA_ENABLE_DEPRECATION_WARNINGS
 
-	PreprocessRenderer->Render();
+	return PreviewRenderTarget_DEPRECATED;
 }
+#endif // WITH_EDITOR
 
-UTexture* UDMXPixelMappingRendererComponent::GetRenderedInputTexture() const
+
+FIntPoint UDMXPixelMappingRendererComponent::GetPixelPosition(int32 InPosition) const
 {
-	return PreprocessRenderer->GetRenderedTexture();
+	// DEPRECATED 5.3
+
+	const int32 YRows = InPosition / MaxDownsampleBufferTargetSize_DEPRECATED.X;
+	return FIntPoint(InPosition % MaxDownsampleBufferTargetSize_DEPRECATED.X, YRows);
 }
 
 void UDMXPixelMappingRendererComponent::CreateOrUpdateDownsampleBufferTarget()
 {
+	// DEPRECATED 5.3
+
 	// Create texture if it does not exists
-	if (DownsampleBufferTarget == nullptr)
-	{	
-		DownsampleBufferTarget = CreateRenderTarget(TEXT("DMXPixelMappingDownsampleBufferTarget"));
+	PRAGMA_DISABLE_DEPRECATION_WARNINGS
+	if (DownsampleBufferTarget_DEPRECATED == nullptr)
+	{
+		DownsampleBufferTarget_DEPRECATED = CreateRenderTarget(TEXT("DMXPixelMappingDownsampleBufferTarget"));
 	}
 
-	const int32 PreviousDownsamplePixelCount = DownsamplePixelCount;
+	const int32 PreviousDownsamplePixelCount = DownsamplePixelCount_DEPRECATED;
 	const int32 TotalDownsamplePixelCount = GetTotalDownsamplePixelCount();
 
-	if (TotalDownsamplePixelCount > 0 && 
+	if (TotalDownsamplePixelCount > 0 &&
 		TotalDownsamplePixelCount != PreviousDownsamplePixelCount)
 	{
 		// Make sure total pixel count less then max texture size MaxDownsampleBufferTargetSize.X * MaxDownsampleBufferTargetSize.Y
-		if (!ensure(TotalDownsamplePixelCount < (MaxDownsampleBufferTargetSize.X * MaxDownsampleBufferTargetSize.Y)))
+		if (!ensure(TotalDownsamplePixelCount < (MaxDownsampleBufferTargetSize_DEPRECATED.X * MaxDownsampleBufferTargetSize_DEPRECATED.Y)))
 		{
 			return;
 		}
 
 		/**
-		 * if total pixel count less then max size x texture high equal 1
-		 * and texture widht dynamic from 1 up to MaxDownsampleBufferTargetSize.X
-		 * |0,1,2,3,4,5,...,n|
-		 */
-		if (TotalDownsamplePixelCount <= MaxDownsampleBufferTargetSize.X)
+			* if total pixel count less then max size x texture high equal 1
+			* and texture widht dynamic from 1 up to MaxDownsampleBufferTargetSize.X
+			* |0,1,2,3,4,5,...,n|
+			*/
+		if (TotalDownsamplePixelCount <= MaxDownsampleBufferTargetSize_DEPRECATED.X)
 		{
 			constexpr uint32 TargetSizeY = 1;
-			DownsampleBufferTarget->ResizeTarget(TotalDownsamplePixelCount, TargetSizeY);
+			DownsampleBufferTarget_DEPRECATED->ResizeTarget(TotalDownsamplePixelCount, TargetSizeY);
 		}
 		/**
 		* if total pixel count more then max size x. At this case it should resize X and Y for buffer texture target
@@ -472,70 +468,88 @@ void UDMXPixelMappingRendererComponent::CreateOrUpdateDownsampleBufferTarget()
 		*/
 		else
 		{
-			const uint32 TargetSizeY = ((TotalDownsamplePixelCount - 1) / MaxDownsampleBufferTargetSize.X) + 1;
-			DownsampleBufferTarget->ResizeTarget(MaxDownsampleBufferTargetSize.X, TargetSizeY);
+			const uint32 TargetSizeY = ((TotalDownsamplePixelCount - 1) / MaxDownsampleBufferTargetSize_DEPRECATED.X) + 1;
+			DownsampleBufferTarget_DEPRECATED->ResizeTarget(MaxDownsampleBufferTargetSize_DEPRECATED.X, TargetSizeY);
 		}
 	}
+
+	PRAGMA_ENABLE_DEPRECATION_WARNINGS
 }
 
+PRAGMA_DISABLE_DEPRECATION_WARNINGS
 void UDMXPixelMappingRendererComponent::AddPixelToDownsampleSet(FDMXPixelMappingDownsamplePixelParamsV2&& InDownsamplePixelParam)
 {
-	const FScopeLock Lock(&DownsampleBufferCS);
-	DownsamplePixelParams.Emplace(InDownsamplePixelParam);
+	// DEPRECATED 5.3
+
+	const FScopeLock Lock(&DownsampleBufferCS_DEPRECATED);
+	DownsamplePixelParams_DEPRECATED.Emplace(InDownsamplePixelParam);
 }
+PRAGMA_ENABLE_DEPRECATION_WARNINGS
 
 int32 UDMXPixelMappingRendererComponent::GetDownsamplePixelNum()
 {
-	const FScopeLock Lock(&DownsampleBufferCS);
-	return DownsamplePixelParams.Num();
+	// DEPRECATED 5.3
+
+	const FScopeLock Lock(&DownsampleBufferCS_DEPRECATED);
+	return DownsamplePixelParams_DEPRECATED.Num();
 }
 
 void UDMXPixelMappingRendererComponent::SetDownsampleBuffer(TArray<FLinearColor>&& InDownsampleBuffer, FIntRect InRect)
 {
+	// DEPRECATED 5.3
+
 	check(IsInRenderingThread());
 
-	if (!bWasEverRendered)
+	if (!bWasEverRendered_DEPRECATED)
 	{
+		PRAGMA_DISABLE_DEPRECATION_WARNINGS
 		for (int32 PixelIndex = 0; PixelIndex < GetTotalDownsamplePixelCount(); PixelIndex++)
 		{
 			ResetColorDownsampleBufferPixel(PixelIndex);
 		}
+		PRAGMA_ENABLE_DEPRECATION_WARNINGS
 
-		bWasEverRendered = true;
+		bWasEverRendered_DEPRECATED = true;
 	}
 
-	FScopeLock ScopeLock(&DownsampleBufferCS);
-	DownsampleBuffer = MoveTemp(InDownsampleBuffer);
+	FScopeLock ScopeLock(&DownsampleBufferCS_DEPRECATED);
+	DownsampleBuffer_DEPRECATED = MoveTemp(InDownsampleBuffer);
 }
 
 bool UDMXPixelMappingRendererComponent::GetDownsampleBufferPixel(const int32 InDownsamplePixelIndex, FLinearColor& OutLinearColor)
 {
-	FScopeLock ScopeLock(&DownsampleBufferCS);
+	// DEPRECATED 5.3
 
+	FScopeLock ScopeLock(&DownsampleBufferCS_DEPRECATED);
 
-	if (!DownsampleBuffer.IsValidIndex(InDownsamplePixelIndex))
+
+	if (!DownsampleBuffer_DEPRECATED.IsValidIndex(InDownsamplePixelIndex))
 	{
 		return false;
 	}
 
-	OutLinearColor = DownsampleBuffer[InDownsamplePixelIndex];
+	OutLinearColor = DownsampleBuffer_DEPRECATED[InDownsamplePixelIndex];
 	return true;
 }
 
 bool UDMXPixelMappingRendererComponent::GetDownsampleBufferPixels(const int32 InDownsamplePixelIndexStart, const int32 InDownsamplePixelIndexEnd, TArray<FLinearColor>& OutLinearColors)
 {
-	FScopeLock ScopeLock(&DownsampleBufferCS);
-	
+	// DEPRECATED 5.3
+
+	FScopeLock ScopeLock(&DownsampleBufferCS_DEPRECATED);
+
 	// Could be out of the range when texture resizing on GPU thread
+	PRAGMA_DISABLE_DEPRECATION_WARNINGS
 	if (!IsPixelRangeValid(InDownsamplePixelIndexStart, InDownsamplePixelIndexEnd))
 	{
 		return false;
 	}
+	PRAGMA_ENABLE_DEPRECATION_WARNINGS
 
-	OutLinearColors.Reset(InDownsamplePixelIndexEnd - InDownsamplePixelIndexStart + 1);
+		OutLinearColors.Reset(InDownsamplePixelIndexEnd - InDownsamplePixelIndexStart + 1);
 	for (int32 PixelIndex = InDownsamplePixelIndexStart; PixelIndex <= InDownsamplePixelIndexEnd; ++PixelIndex)
 	{
-		OutLinearColors.Add(DownsampleBuffer[PixelIndex]);
+		OutLinearColors.Add(DownsampleBuffer_DEPRECATED[PixelIndex]);
 	}
 
 	return true;
@@ -543,11 +557,13 @@ bool UDMXPixelMappingRendererComponent::GetDownsampleBufferPixels(const int32 In
 
 bool UDMXPixelMappingRendererComponent::ResetColorDownsampleBufferPixel(const int32 InDownsamplePixelIndex)
 {
-	FScopeLock ScopeLock(&DownsampleBufferCS);
+	// DEPRECATED 5.3
+
+	FScopeLock ScopeLock(&DownsampleBufferCS_DEPRECATED);
 
-	if (DownsampleBuffer.IsValidIndex(InDownsamplePixelIndex))
+	if (DownsampleBuffer_DEPRECATED.IsValidIndex(InDownsamplePixelIndex))
 	{
-		DownsampleBuffer[InDownsamplePixelIndex] = FLinearColor::Black;
+		DownsampleBuffer_DEPRECATED[InDownsamplePixelIndex] = FLinearColor::Black;
 		return true;
 	}
 
@@ -556,17 +572,21 @@ bool UDMXPixelMappingRendererComponent::ResetColorDownsampleBufferPixel(const in
 
 bool UDMXPixelMappingRendererComponent::ResetColorDownsampleBufferPixels(const int32 InDownsamplePixelIndexStart, const int32 InDownsamplePixelIndexEnd)
 {
-	FScopeLock ScopeLock(&DownsampleBufferCS);
+	// DEPRECATED 5.3
+
+	FScopeLock ScopeLock(&DownsampleBufferCS_DEPRECATED);
 
 	// Could be out of the range when texture resizing on GPU thread
+	PRAGMA_DISABLE_DEPRECATION_WARNINGS
 	if (!IsPixelRangeValid(InDownsamplePixelIndexStart, InDownsamplePixelIndexEnd))
 	{
 		return false;
 	}
+	PRAGMA_ENABLE_DEPRECATION_WARNINGS
 
 	for (int32 PixelIndex = InDownsamplePixelIndexStart; PixelIndex <= InDownsamplePixelIndexEnd; ++PixelIndex)
 	{
-		DownsampleBuffer[PixelIndex] = FLinearColor::Black;
+		DownsampleBuffer_DEPRECATED[PixelIndex] = FLinearColor::Black;
 	}
 
 	return true;
@@ -574,33 +594,46 @@ bool UDMXPixelMappingRendererComponent::ResetColorDownsampleBufferPixels(const i
 
 void UDMXPixelMappingRendererComponent::EmptyDownsampleBuffer()
 {
-	FScopeLock ScopeLock(&DownsampleBufferCS);
+	// DEPRECATED 5.3
+
+	FScopeLock ScopeLock(&DownsampleBufferCS_DEPRECATED);
 
-	DownsampleBuffer.Empty();
+	DownsampleBuffer_DEPRECATED.Empty();
 }
 
-bool UDMXPixelMappingRendererComponent::IsPixelRangeValid(const int32 InDownsamplePixelIndexStart, const int32 InDownsamplePixelIndexEnd) const
-{
-	FScopeLock ScopeLock(&DownsampleBufferCS);
+#if WITH_EDITOR
+TSharedRef<SWidget> UDMXPixelMappingRendererComponent::TakeWidget()
+{	
+	// DEPRECATED 5.3
 
-	if (InDownsamplePixelIndexEnd >= InDownsamplePixelIndexStart &&
-		DownsampleBuffer.IsValidIndex(InDownsamplePixelIndexStart) &&
-		DownsampleBuffer.IsValidIndex(InDownsamplePixelIndexEnd))
+	if (!ComponentsCanvas_DEPRECATED.IsValid())
 	{
-		return true;
+		ComponentsCanvas_DEPRECATED =
+			SNew(SConstraintCanvas);
 	}
 
-	return false;
+	PRAGMA_DISABLE_DEPRECATION_WARNINGS
+	ForEachChild([&](UDMXPixelMappingBaseComponent* InComponent) {
+			if (UDMXPixelMappingOutputComponent* Component = Cast<UDMXPixelMappingOutputComponent>(InComponent))
+			{
+				// Build all child DMX pixel mapping slots
+				Component->BuildSlot(ComponentsCanvas_DEPRECATED.ToSharedRef());
+			}
+		}, true);
+	PRAGMA_ENABLE_DEPRECATION_WARNINGS
+
+	return ComponentsCanvas_DEPRECATED.ToSharedRef();
 }
+#endif // WITH_EDITOR
 
 int32 UDMXPixelMappingRendererComponent::GetTotalDownsamplePixelCount()
 {
-	SCOPE_CYCLE_COUNTER(STAT_GetTotalDownsamplePixelCount);
+	// DEPRECATED 5.3
 
-	FScopeLock ScopeLock(&DownsampleBufferCS);
+	FScopeLock ScopeLock(&DownsampleBufferCS_DEPRECATED);
 
 	// Reset pixel counter
-	DownsamplePixelCount = 0;
+	DownsamplePixelCount_DEPRECATED = 0;
 
 	// Count all pixels
 	constexpr bool bIsRecursive = true;
@@ -609,14 +642,60 @@ int32 UDMXPixelMappingRendererComponent::GetTotalDownsamplePixelCount()
 			// If that is screen component
 			if (UDMXPixelMappingScreenComponent* ScreenComponent = Cast<UDMXPixelMappingScreenComponent>(InComponent))
 			{
-				DownsamplePixelCount += (ScreenComponent->NumXCells * ScreenComponent->NumYCells);
+				DownsamplePixelCount_DEPRECATED += (ScreenComponent->NumXCells * ScreenComponent->NumYCells);
 			}
 			// If that is single pixel component
 			else if (Cast<UDMXPixelMappingOutputDMXComponent>(InComponent))
 			{
-				DownsamplePixelCount++;
+				DownsamplePixelCount_DEPRECATED++;
 			}
 		}, bIsRecursive);
 
-	return DownsamplePixelCount;
+	return DownsamplePixelCount_DEPRECATED;
+}
+
+bool UDMXPixelMappingRendererComponent::IsPixelRangeValid(const int32 InDownsamplePixelIndexStart, const int32 InDownsamplePixelIndexEnd) const
+{
+	// DEPRECATED 5.3
+
+	FScopeLock ScopeLock(&DownsampleBufferCS_DEPRECATED);
+
+	if (InDownsamplePixelIndexEnd >= InDownsamplePixelIndexStart &&
+		DownsampleBuffer_DEPRECATED.IsValidIndex(InDownsamplePixelIndexStart) &&
+		DownsampleBuffer_DEPRECATED.IsValidIndex(InDownsamplePixelIndexEnd))
+	{
+		return true;
+	}
+
+	return false;
+}
+
+#if WITH_EDITOR
+void UDMXPixelMappingRendererComponent::ResizePreviewRenderTarget(uint32 InSizeX, uint32 InSizeY)
+{	
+	// DEPRECATED 5.3
+	PRAGMA_DISABLE_DEPRECATION_WARNINGS
+	UTextureRenderTarget2D* Target = GetPreviewRenderTarget();
+
+	if ((InSizeX > 0 && InSizeY > 0) && (Target->SizeX != InSizeX || Target->SizeY != InSizeY))
+	{
+		check(Target);
+		Target->ResizeTarget(InSizeX, InSizeY);
+		Target->UpdateResourceImmediate();
+	}
+	PRAGMA_ENABLE_DEPRECATION_WARNINGS
+}
+#endif // WITH_EDITOR
+
+UTextureRenderTarget2D* UDMXPixelMappingRendererComponent::CreateRenderTarget(const FName& InBaseName)
+{
+	// DEPRECATED 5.3
+
+	const FName TargetName = MakeUniqueObjectName(this, UTextureRenderTarget2D::StaticClass(), InBaseName);
+	UTextureRenderTarget2D* RenderTarget = NewObject<UTextureRenderTarget2D>(this, TargetName);
+	RenderTarget->ClearColor = ClearTextureColor_DEPRECATED;
+	constexpr bool bInForceLinearGamma = false;
+	RenderTarget->InitCustomFormat(GetSize().X, GetSize().Y, EPixelFormat::PF_B8G8R8A8, bInForceLinearGamma);
+
+	return RenderTarget;
 }
diff --git a/Engine/Plugins/VirtualProduction/DMX/DMXPixelMapping/Source/DMXPixelMappingRuntime/Private/Components/DMXPixelMappingScreenComponent.cpp b/Engine/Plugins/VirtualProduction/DMX/DMXPixelMapping/Source/DMXPixelMappingRuntime/Private/Components/DMXPixelMappingScreenComponent.cpp
--- a/Engine/Plugins/VirtualProduction/DMX/DMXPixelMapping/Source/DMXPixelMappingRuntime/Private/Components/DMXPixelMappingScreenComponent.cpp
+++ b/Engine/Plugins/VirtualProduction/DMX/DMXPixelMapping/Source/DMXPixelMappingRuntime/Private/Components/DMXPixelMappingScreenComponent.cpp
@@ -4,7 +4,6 @@
 
 #include "DMXPixelMappingUtils.h"
 #include "Components/DMXPixelMappingRendererComponent.h"
-#include "IDMXPixelMappingRenderer.h"
 #include "IO/DMXOutputPort.h"
 #include "IO/DMXOutputPortReference.h"
 #include "IO/DMXPortManager.h"
@@ -265,16 +264,15 @@ UDMXPixelMappingRendererComponent* UDMXPixelMappingScreenComponent::GetRendererC
 
 void UDMXPixelMappingScreenComponent::ResetDMX()
 {
-	UDMXPixelMappingRendererComponent* RendererComponent = GetRendererComponent();
-	if (ensure(RendererComponent))
-	{
-	RendererComponent->ResetColorDownsampleBufferPixels(PixelDownsamplePositionRange.Key, PixelDownsamplePositionRange.Value);
+	//RenderElement->SetColor(FLinearColor::Black);
+
 	SendDMX();
 }
-}
 
 void UDMXPixelMappingScreenComponent::SendDMX()
 {
+	PRAGMA_DISABLE_DEPRECATION_WARNINGS
+
 	SCOPE_CYCLE_COUNTER(STAT_DMXPixelMaping_SendScreen);
 
 	UDMXPixelMappingRendererComponent* RendererComponent = GetRendererComponent();
@@ -372,10 +370,15 @@ void UDMXPixelMappingScreenComponent::SendDMX()
 			SendDMXIndex++;
 		}
 	}
+
+	PRAGMA_ENABLE_DEPRECATION_WARNINGS
 }
 
 void UDMXPixelMappingScreenComponent::QueueDownsample()
 {
+	// DEPRECATED 5.3
+	PRAGMA_DISABLE_DEPRECATION_WARNINGS
+
 	// Queue pixels into the downsample rendering
 	UDMXPixelMappingRendererComponent* RendererComponent = GetRendererComponent();
 	if (!ensure(RendererComponent))
@@ -427,15 +430,13 @@ void UDMXPixelMappingScreenComponent::QueueDownsample()
 
 	// End of downsample index
 	PixelDownsamplePositionRange.Value = PixelDownsamplePositionRange.Key + IterationCount;
+
+	PRAGMA_ENABLE_DEPRECATION_WARNINGS
 }
 
 void UDMXPixelMappingScreenComponent::RenderWithInputAndSendDMX()
 {
-	if (UDMXPixelMappingRendererComponent* RendererComponent = GetFirstParentByClass<UDMXPixelMappingRendererComponent>(this))
-	{
-		RendererComponent->RendererInputTexture();
-	}
-
+	// DEPRECATED 5.3
 	RenderAndSendDMX();
 }
 
diff --git a/Engine/Plugins/VirtualProduction/DMX/DMXPixelMapping/Source/DMXPixelMappingRuntime/Public/Components/DMXPixelMappingBaseComponent.h b/Engine/Plugins/VirtualProduction/DMX/DMXPixelMapping/Source/DMXPixelMappingRuntime/Public/Components/DMXPixelMappingBaseComponent.h
--- a/Engine/Plugins/VirtualProduction/DMX/DMXPixelMapping/Source/DMXPixelMappingRuntime/Public/Components/DMXPixelMappingBaseComponent.h
+++ b/Engine/Plugins/VirtualProduction/DMX/DMXPixelMapping/Source/DMXPixelMappingRuntime/Public/Components/DMXPixelMappingBaseComponent.h
@@ -3,10 +3,7 @@
 #pragma once
 
 #include "DMXPixelMappingRuntimeCommon.h"
-
-#include "Tickable.h"
 #include "UObject/Object.h"
-#include "Stats/Stats.h"
 
 #include "DMXPixelMappingBaseComponent.generated.h"
 
@@ -23,7 +20,6 @@ class UDMXPixelMappingRendererComponent;
 UCLASS(BlueprintType, NotBlueprintable, Abstract)
 class DMXPIXELMAPPINGRUNTIME_API UDMXPixelMappingBaseComponent
 	: public UObject
-	, public FTickableGameObject
 {
 	GENERATED_BODY()
 
@@ -32,9 +28,6 @@ class DMXPIXELMAPPINGRUNTIME_API UDMXPixelMappingBaseComponent
 	DECLARE_EVENT_FourParams(UDMXPixelMappingBaseComponent, FDMXPixelMappingOnComponentRenamed, UDMXPixelMapping* /** PixelMapping */, UDMXPixelMappingBaseComponent* /** RenamedComponent */, UObject* /** OldOuter */, const FName /** OldName */);
 
 public:
-	/** Public constructor */
-	UDMXPixelMappingBaseComponent();
-
 	/** Gets an Event broadcast when a component was added */
 	static FDMXPixelMappingOnComponentAdded& GetOnComponentAdded();
 
@@ -66,14 +59,6 @@ public:
 	*/
 	virtual const FName& GetNamePrefix();
 
-	// ~Begin FTickableGameObject interface
-	virtual void Tick(float DeltaTime) override {}
-	virtual TStatId GetStatId() const override;
-	virtual bool IsTickableInEditor() const override { return true; }
-	virtual bool IsTickableWhenPaused() const override { return true; }
-	virtual bool IsTickable() const override { return false; }
-	// ~End FTickableGameObject interface
-
 	/*------------------------------------------
 		UDMXPixelMappingBaseComponent interface
 	--------------------------------------------*/
diff --git a/Engine/Plugins/VirtualProduction/DMX/DMXPixelMapping/Source/DMXPixelMappingRuntime/Public/Components/DMXPixelMappingFixtureGroupComponent.h b/Engine/Plugins/VirtualProduction/DMX/DMXPixelMapping/Source/DMXPixelMappingRuntime/Public/Components/DMXPixelMappingFixtureGroupComponent.h
--- a/Engine/Plugins/VirtualProduction/DMX/DMXPixelMapping/Source/DMXPixelMappingRuntime/Public/Components/DMXPixelMappingFixtureGroupComponent.h
+++ b/Engine/Plugins/VirtualProduction/DMX/DMXPixelMapping/Source/DMXPixelMappingRuntime/Public/Components/DMXPixelMappingFixtureGroupComponent.h
@@ -48,11 +48,16 @@ public:
 	virtual bool IsExposedToTemplate() { return true; }
 	virtual const FText GetPaletteCategory() override;
 #endif // WITH_EDITOR
-	virtual void QueueDownsample() override;
 	virtual void SetPosition(const FVector2D& NewPosition) override;
 	virtual void SetSize(const FVector2D& NewSize) override;
 	//~ End UDMXPixelMappingOutputComponent implementation
 
+	/** Returns a delegate broadcasted when the DMX Library changed */
+	FSimpleMulticastDelegate& GetOnDMXLibraryChanged() { return OnDMXLibraryChangedDelegate; }
+
+	UE_DEPRECATED(5.3, "Please use UDMXPixelMappingPixelMapRenderer to render the pixel map")
+	virtual void QueueDownsample() override;
+
 	UPROPERTY(EditAnywhere, BlueprintReadOnly, Category = "Fixture List")
 	TObjectPtr<UDMXLibrary> DMXLibrary;
 
@@ -60,9 +65,6 @@ public:
 	UPROPERTY(EditAnywhere, Instanced, Category = "Layout")
 	TObjectPtr<UDMXPixelMappingLayoutScript> LayoutScript;
 
-	/** Returns a delegate broadcasted when the DMX Library changed */
-	FSimpleMulticastDelegate& GetOnDMXLibraryChanged() { return OnDMXLibraryChangedDelegate; }
-
 private:
 	/** Handles changes in position */
 	void HandlePositionChanged();
diff --git a/Engine/Plugins/VirtualProduction/DMX/DMXPixelMapping/Source/DMXPixelMappingRuntime/Public/Components/DMXPixelMappingFixtureGroupItemComponent.h b/Engine/Plugins/VirtualProduction/DMX/DMXPixelMapping/Source/DMXPixelMappingRuntime/Public/Components/DMXPixelMappingFixtureGroupItemComponent.h
--- a/Engine/Plugins/VirtualProduction/DMX/DMXPixelMapping/Source/DMXPixelMappingRuntime/Public/Components/DMXPixelMappingFixtureGroupItemComponent.h
+++ b/Engine/Plugins/VirtualProduction/DMX/DMXPixelMapping/Source/DMXPixelMappingRuntime/Public/Components/DMXPixelMappingFixtureGroupItemComponent.h
@@ -14,13 +14,12 @@
 
 enum class EDMXColorMode : uint8;
 struct FDMXAttributeName;
-class UDMXModulator;
-class UDMXPixelMappingColorSpace;
-
 class FEditPropertyChain;
 class STextBlock;
 class UTextureRenderTarget2D;
-
+class UDMXModulator;
+class UDMXPixelMappingColorSpace;
+namespace UE::DMXPixelMapping::Rendering::PixelMapRenderer { class FPixelMapRenderElement; }
 
 /** 
  * A component that holds a single Fixture Patch in the Pixel Mapping, and actually sends DMX.
@@ -38,13 +37,13 @@ public:
 protected:
 	//~ Begin UObject interface
 	virtual void Serialize(FArchive& Ar) override;
+	virtual void PostInitProperties() override;
 	virtual void PostLoad() override;
 	virtual void BeginDestroy() override;
 #if WITH_EDITOR
 	virtual void PostEditChangeProperty(FPropertyChangedEvent& PropertyChangedEvent) override;
 	virtual void PostEditChangeChainProperty(FPropertyChangedChainEvent& PropertyChangedChainEvent) override;
 #endif // WITH_EDITOR
-
 	//~ End UObject interface
 
 public:
@@ -60,15 +59,12 @@ public:
 	virtual bool IsVisible() const override;
 #endif // WITH_EDITOR	
 	virtual bool IsOverParent() const override;
-	virtual void QueueDownsample() override;
-	virtual int32 GetDownsamplePixelIndex() const override { return DownsamplePixelIndex; }
 	virtual void SetPosition(const FVector2D& NewPosition) override;
 	virtual void SetSize(const FVector2D& NewSize) override;
 	//~ End UDMXPixelMappingOutputComponent implementation
 
-	//~ Begin UDMXPixelMappingOutputDMXComponent implementation
-	virtual void RenderWithInputAndSendDMX() override;
-	//~ End UDMXPixelMappingOutputDMXComponent implementation
+	/** Returns the pixel map render element for the matrix cell */
+	TSharedRef<UE::DMXPixelMapping::Rendering::FPixelMapRenderElement> GetOrCreatePixelMapRenderElement();
 
 	/** Check if a Component can be moved under another one (used for copy/move/duplicate) */
 	virtual bool CanBeMovedTo(const UDMXPixelMappingBaseComponent* Component) const override;
@@ -96,6 +92,9 @@ public:
 	int32 DownsamplePixelIndex;
 
 private:
+	/** Updates the rendered element */
+	void UpdateRenderElement();
+
 #if WITH_EDITOR
 	/** Called post edit change properties of the color space member */
 	void OnColorSpacePostEditChangeProperties(FPropertyChangedEvent& PropertyChangedEvent);
@@ -104,11 +103,29 @@ private:
 	/** Helper that returns the renderer component this component belongs to */
 	UDMXPixelMappingRendererComponent* GetRendererComponent() const;
 
+	/** Pixel map render element for the matrix cell */
+	TSharedPtr<UE::DMXPixelMapping::Rendering::FPixelMapRenderElement> PixelMapRenderElement;
+
 
 	//////////////////////
 	// Deprecated Members
 	
 public:
+	//~ Begin UDMXPixelMappingOutputComponent implementation
+	UE_DEPRECATED(5.3, "No longer in use. Please use UDMXPixelMappingPixelMapRenderer to render the pixel map")
+	virtual void QueueDownsample() override;
+	UE_DEPRECATED(5.3, "No longer in use. Please use UDMXPixelMappingPixelMapRenderer to render the pixel map")
+	virtual int32 GetDownsamplePixelIndex() const override { return DownsamplePixelIndex; }
+	//~ End UDMXPixelMappingOutputComponent implementation
+
+
+	//~ Begin UDMXPixelMappingOutputDMXComponent implementation
+	PRAGMA_DISABLE_DEPRECATION_WARNINGS
+	UE_DEPRECATED(5.3, "Deprecated for performance reasons. Instead use 'Get DMX Pixel Mapping Renderer Component' and Render only once each tick.")
+	virtual void RenderWithInputAndSendDMX() override;
+	PRAGMA_ENABLE_DEPRECATION_WARNINGS
+	//~ End UDMXPixelMappingOutputDMXComponent implementation
+
 	/** DEPRECATED 5.2 */
 	UE_DEPRECATED(5.2, "Instead use UDMXPixelMappingColorSpace::Update and UDMXPixelMappingColorSpace::GetAttributeNameToValueMap using the ColorSpace member of this class.")
 	TMap<FDMXAttributeName, float> CreateAttributeValues() const;
diff --git a/Engine/Plugins/VirtualProduction/DMX/DMXPixelMapping/Source/DMXPixelMappingRuntime/Public/Components/DMXPixelMappingMatrixCellComponent.h b/Engine/Plugins/VirtualProduction/DMX/DMXPixelMapping/Source/DMXPixelMappingRuntime/Public/Components/DMXPixelMappingMatrixCellComponent.h
--- a/Engine/Plugins/VirtualProduction/DMX/DMXPixelMapping/Source/DMXPixelMappingRuntime/Public/Components/DMXPixelMappingMatrixCellComponent.h
+++ b/Engine/Plugins/VirtualProduction/DMX/DMXPixelMapping/Source/DMXPixelMappingRuntime/Public/Components/DMXPixelMappingMatrixCellComponent.h
@@ -4,14 +4,19 @@
 
 #include "Components/DMXPixelMappingOutputDMXComponent.h"
 #include "DMXAttribute.h"
+#include "DMXTypes.h"
 #include "Library/DMXEntityReference.h"
+#include "PixelMapRenderElement.h"
 
 #include "DMXPixelMappingMatrixCellComponent.generated.h"
 
 
+enum class EDMXColorMode : uint8;
+namespace UE::DMXPixelMapping::Rendering::PixelMapRenderer { class FPixelMapRenderElement; }
 class SUniformGridPanel;
+class UDMXPixelMappingColorSpace;
 class UTextureRenderTarget2D;
-enum class EDMXColorMode : uint8;
+
 
 /**
  * Matrix pixel component
@@ -30,6 +35,9 @@ public:
 	UDMXPixelMappingMatrixCellComponent();
 
 	//~ Begin UObject implementation
+	virtual void PostInitProperties() override;
+	virtual void PostLoad() override;
+	virtual void BeginDestroy() override;
 #if WITH_EDITOR
 	virtual void PostEditChangeProperty(FPropertyChangedEvent& PropertyChangedEvent) override;
 #endif // WITH_EDITOR
@@ -56,9 +64,10 @@ public:
 	virtual void SetSize(const FVector2D& NewSize) override;
 	//~ End UDMXPixelMappingOutputComponent implementation
 
-public:
 	//~ Begin UDMXPixelMappingOutputDMXComponent implementation
+	PRAGMA_DISABLE_DEPRECATION_WARNINGS
 	virtual void RenderWithInputAndSendDMX() override;
+	PRAGMA_ENABLE_DEPRECATION_WARNINGS
 	//~ End UDMXPixelMappingOutputDMXComponent implementation
 
 	void SetCellCoordinate(FIntPoint InCellCoordinate);
@@ -67,21 +76,25 @@ public:
 	/** Check if a Component can be moved under another one (used for copy/move/duplicate) */
 	virtual bool CanBeMovedTo(const UDMXPixelMappingBaseComponent* Component) const override;
 
-private:
-	/** Helper that returns the renderer component this component belongs to */
-	UDMXPixelMappingRendererComponent* GetRendererComponent() const;
+	/** Returns the pixel map render element for the matrix cell */
+	TSharedRef<UE::DMXPixelMapping::Rendering::FPixelMapRenderElement> GetOrCreatePixelMapRenderElement();
+
+	/** Gets the color of this cell */
+	FLinearColor GetPixelMapColor() const;
 
-public:
 	UPROPERTY(VisibleAnywhere, BlueprintReadOnly, Category = "Cell Settings")
 	int32 CellID;
 
-	UPROPERTY()
-	FDMXEntityFixturePatchRef FixturePatchMatrixRef_DEPRECATED;
+private:
+	/** Updates the rendered element. Automatically called on property changes */
+	void UpdateRenderElement();
 
-	/** Creates attribute values from current data */
-	TMap<FDMXAttributeName, float> CreateAttributeValues() const;
+	/** Helper that returns the renderer component this component belongs to */
+	UDMXPixelMappingRendererComponent* GetRendererComponent() const;
+
+	/** Pixel map render element for the matrix cell */
+	TSharedPtr<UE::DMXPixelMapping::Rendering::FPixelMapRenderElement> PixelMapRenderElement;
 
-private:
 	UPROPERTY()
 	FIntPoint CellCoordinate;
 
diff --git a/Engine/Plugins/VirtualProduction/DMX/DMXPixelMapping/Source/DMXPixelMappingRuntime/Public/Components/DMXPixelMappingMatrixComponent.h b/Engine/Plugins/VirtualProduction/DMX/DMXPixelMapping/Source/DMXPixelMappingRuntime/Public/Components/DMXPixelMappingMatrixComponent.h
--- a/Engine/Plugins/VirtualProduction/DMX/DMXPixelMapping/Source/DMXPixelMappingRuntime/Public/Components/DMXPixelMappingMatrixComponent.h
+++ b/Engine/Plugins/VirtualProduction/DMX/DMXPixelMapping/Source/DMXPixelMappingRuntime/Public/Components/DMXPixelMappingMatrixComponent.h
@@ -66,7 +66,6 @@ public:
 	virtual const FText GetPaletteCategory() override;
 #endif // WITH_EDITOR
 	virtual bool IsOverParent() const override;
-	virtual void QueueDownsample() override;
 	virtual void SetPosition(const FVector2D& NewPosition) override;
 	virtual void SetSize(const FVector2D& NewSize) override;
 	// ~End UDMXPixelMappingOutputComponent interface
@@ -80,6 +79,9 @@ public:
 	/** Handles changes in size or in matrix */
 	void HandleMatrixChanged();
 
+	UE_DEPRECATED(5.3, "Please use UDMXPixelMappingPixelMapRenderer to render the pixel map")
+	virtual void QueueDownsample() override;
+
 protected:
 	/** Called when the fixture type in use changed */
 	void OnFixtureTypeChanged(const UDMXEntityFixtureType* FixtureType);
diff --git a/Engine/Plugins/VirtualProduction/DMX/DMXPixelMapping/Source/DMXPixelMappingRuntime/Public/Components/DMXPixelMappingOutputComponent.h b/Engine/Plugins/VirtualProduction/DMX/DMXPixelMapping/Source/DMXPixelMappingRuntime/Public/Components/DMXPixelMappingOutputComponent.h
--- a/Engine/Plugins/VirtualProduction/DMX/DMXPixelMapping/Source/DMXPixelMappingRuntime/Public/Components/DMXPixelMappingOutputComponent.h
+++ b/Engine/Plugins/VirtualProduction/DMX/DMXPixelMapping/Source/DMXPixelMappingRuntime/Public/Components/DMXPixelMappingOutputComponent.h
@@ -11,15 +11,17 @@
 #include "DMXPixelMappingOutputComponent.generated.h"
 
 struct FDMXPixelMappingLayoutToken;
+class SBox;
 class UDMXEntityFixturePatch;
 class UDMXPixelMappingRendererComponent;
+namespace UE::DMXPixelMapping::Rendering::PixelMapRenderer { class FPixelMapRenderElement; }
+
 
 #if WITH_EDITOR
 enum class EDMXPixelMappingComponentLabelAlignment : uint8;
 class FDMXPixelMappingComponentWidget;
 #endif
 
-class SBox;
 
 
 /** Enum that defines the quality of how pixels are rendered */
@@ -108,16 +110,18 @@ public:
 	virtual bool OverlapsComponent(UDMXPixelMappingOutputComponent* Other) const;
 
 	/** Get pixel index in downsample texture */
+	UE_DEPRECATED(5.3, "Please use UDMXPixelMappingPixelMapRenderer to render the pixel map")
 	virtual int32 GetDownsamplePixelIndex() const { return 0; }
 
 	/** Queue rendering to downsample rendering target */
+	UE_DEPRECATED(5.3, "Please use UDMXPixelMappingPixelMapRenderer to render the pixel map")
 	virtual void QueueDownsample() {}
 
 	/** Sets the position */
 	virtual void SetPosition(const FVector2D& NewPosition);
 
-	/** Returns the position. Note this is the position before a layout has been applied, use GetLayoutedPosition() to get the final position.. */
-	FVector2D GetPosition() const { return FVector2D(PositionX, PositionY); }
+	/** Returns the position */
+	FVector2D GetPosition() const;
 
 	/** Sets the size */
 	virtual void SetSize(const FVector2D& NewSize);
@@ -125,6 +129,9 @@ public:
 	/** Get the size */
 	FVector2D GetSize() const { return FVector2D(SizeX, SizeY); }
 
+	/** Invalidates the pixel map, effectively causing the renderer component to aquire a new pixel map */
+	void InvalidatePixelMapRenderer();
+
 	/** Helper that returns render component if available */
 	UDMXPixelMappingRendererComponent* FindRendererComponent() const;
 	/** Updates children to match the size of this instance */
diff --git a/Engine/Plugins/VirtualProduction/DMX/DMXPixelMapping/Source/DMXPixelMappingRuntime/Public/Components/DMXPixelMappingOutputDMXComponent.h b/Engine/Plugins/VirtualProduction/DMX/DMXPixelMapping/Source/DMXPixelMappingRuntime/Public/Components/DMXPixelMappingOutputDMXComponent.h
--- a/Engine/Plugins/VirtualProduction/DMX/DMXPixelMapping/Source/DMXPixelMappingRuntime/Public/Components/DMXPixelMappingOutputDMXComponent.h
+++ b/Engine/Plugins/VirtualProduction/DMX/DMXPixelMapping/Source/DMXPixelMappingRuntime/Public/Components/DMXPixelMappingOutputDMXComponent.h
@@ -17,8 +17,9 @@ class DMXPIXELMAPPINGRUNTIME_API UDMXPixelMappingOutputDMXComponent
 
 public:
 	/** Render input texture for downsample texture, donwsample and send DMX for this component */
-	UFUNCTION(BlueprintCallable, Category = "DMX|PixelMapping")
-	virtual void RenderWithInputAndSendDMX() PURE_VIRTUAL(UDMXPixelMappingOutputComponent::RenderWithInputAndSendDMX);
+	UE_DEPRECATED(5.3, "Deprecated for performance reasons. Instead use 'Get DMX Pixel Mapping Renderer Component' and Render only once each tick.")
+	UFUNCTION(BlueprintCallable, Category = "DMX|PixelMapping", Meta = (DeprecatedFunction, DeprecationMessage = "Deprecated for performance reasons. Instead use 'Get DMX Pixel Mapping Renderer Component' and Render only once each tick"))
+	virtual void RenderWithInputAndSendDMX() PURE_VIRTUAL(UDMXPixelMappingOutputDMXComponent::RenderWithInputAndSendDMX);
 
 	/** Check if a Component can be moved under another one (used for copy/move/duplicate) */
 	virtual bool CanBeMovedTo(const UDMXPixelMappingBaseComponent* Component) const override
diff --git a/Engine/Plugins/VirtualProduction/DMX/DMXPixelMapping/Source/DMXPixelMappingRuntime/Public/Components/DMXPixelMappingRendererComponent.h b/Engine/Plugins/VirtualProduction/DMX/DMXPixelMapping/Source/DMXPixelMappingRuntime/Public/Components/DMXPixelMappingRendererComponent.h
--- a/Engine/Plugins/VirtualProduction/DMX/DMXPixelMapping/Source/DMXPixelMappingRuntime/Public/Components/DMXPixelMappingRendererComponent.h
+++ b/Engine/Plugins/VirtualProduction/DMX/DMXPixelMapping/Source/DMXPixelMappingRuntime/Public/Components/DMXPixelMappingRendererComponent.h
@@ -2,8 +2,8 @@
 
 #pragma once
 
-#include "DMXPixelMappingPreprocessRenderer.h"
 #include "Components/DMXPixelMappingOutputComponent.h"
+#include "DMXPixelMappingPreprocessRenderer.h"
 #include "IDMXPixelMappingRenderer.h"
 #include "Library/DMXEntityReference.h"
 #include "Templates/SubclassOf.h"
@@ -12,13 +12,14 @@
 
 enum class EDMXPixelMappingRendererType : uint8;
 class UDMXPixelMappingLayoutScript;
+class UDMXPixelMappingPixelMapRenderer;
 class UDMXPixelMappingPreprocessRenderer;
-
 class UMaterialInterface;
 class UTexture;
-class UUserWidget;
 class UTextureRenderTarget2D;
+class UUserWidget;
 class UWorld;
+namespace UE::DMXPixelMapping::Rendering { class FPixelMapRenderElement; }
 
 
 /** 
@@ -35,14 +36,17 @@ public:
 	UDMXPixelMappingRendererComponent();
 
 	//~ Begin UObject implementation
+protected:
 	virtual void Serialize(FArchive& Ar) override;
 	virtual void PostInitProperties() override;
 	virtual void PostLoad() override;
 #if WITH_EDITOR
+	virtual void PostEditUndo() override;
 	virtual void PostEditChangeChainProperty(FPropertyChangedChainEvent& PropertyChangedChainEvent) override;
 #endif // WITH_EDITOR
 	//~ End UObject implementation
 
+public:
 	//~ Begin UDMXPixelMappingBaseComponent implementation
 	virtual const FName& GetNamePrefix() override;
 	virtual bool CanBeMovedTo(const UDMXPixelMappingBaseComponent* Component) const override;
@@ -52,108 +56,30 @@ public:
 	virtual void RenderAndSendDMX() final;
 	//~ End UDMXPixelMappingBaseComponent implementation
 
-#if WITH_EDITOR
-	/** Render all downsample pixel for editor preview texture */
-	void RenderEditorPreviewTexture();
-	
-	/** Get target for preview, create new one if does not exists. */
-	UTextureRenderTarget2D* GetPreviewRenderTarget();
-#endif
-
-	/** Gets the rendered input texture, or nullptr if no input texture is currently rendered */
-	UTexture* GetRenderedInputTexture() const;
-
-	/**
-	 * Get pixel position in downsample buffer target based on pixel index
-	 *
-	 * @param InIndex Index of the pixel in buffer texture.
-	 * @return FIntPoint X and Y position of the pixel in texture
-	 */
-	FIntPoint GetPixelPosition(int32 InIndex) const;
-
-	/** Get active world. It could be editor or build world */
-	UWorld* GetWorld() const;
+	/** Returns the pixel mapping preprocess renderer */
+	UDMXPixelMappingPreprocessRenderer* GetPreprocessRenderer() const { return PreprocessRenderer; }
 
-	/** Get renderer interface */
-	const TSharedPtr<IDMXPixelMappingRenderer>& GetRenderer() { return PixelMappingRenderer; }
+	/** Initializes the preprocess renderer from the current render type and input source */
+	void InvalidatePreprocessRenderer();
 
-#if WITH_EDITOR
 	/**
-	 * Take of container widget which is holds widget for all child components.
+	 * Invalidates the pixel map, effectively causing a rebuild of the mapping on the next Render call.
+	 * Does not rebuild the renderer self. Use InitializeRenderer to rebuild the renderers.
 	 */
-	UE_DEPRECATED(5.1, "Pixel Mapping Components no longer hold their own widget, in an effort to separate Views from Data.")
-	TSharedRef<SWidget> TakeWidget();
-#endif // WITH_EDITOR
-
-	/*----------------------------------
-		Blueprint interface
-	----------------------------------*/
-
-	/** Render input texture for downsampling */
-	UFUNCTION(BlueprintCallable, Category = "DMX|PixelMapping")
-	void RendererInputTexture();
-
-	/** Create or update size of buffer target for rendering downsample pixels */
-	void CreateOrUpdateDownsampleBufferTarget();
-
-	/** 
-	 * Add pixel params for downsampling set
-	 *
-	 * @param InDownsamplePixelParam pixel rendering params
-	 */
-	void AddPixelToDownsampleSet(FDMXPixelMappingDownsamplePixelParamsV2&& InDownsamplePixelParam);
+	void InvalidatePixelMapRenderer();
 
-	/** Get amount of downsample pixels */
-	int32 GetDownsamplePixelNum();
-
-	/**
-	 * Pass the downsample CPU buffer from Render Thread to Game Thread and store 
-	 * 
-	 * @param InDownsampleBuffer CPU buffer
-	 * @param InRect buffer X and Y dimension
-	 */
-	void SetDownsampleBuffer(TArray<FLinearColor>&& InDownsampleBuffer, FIntRect InRect);
-
-	/** Get Pixel color by given downsample pixel index. Returns false if no color value could be acquired */
-	bool GetDownsampleBufferPixel(const int32 InDownsamplePixelIndex, FLinearColor& OutLinearColor);
-
-	/** Get Pixels color by given downsample pixel range. Returns false if no color values could be acquired */
-	bool GetDownsampleBufferPixels(const int32 InDownsamplePixelIndexStart, const int32 InDownsamplePixelIndexEnd, TArray<FLinearColor>& OutLinearColors);
-
-	/** Reset the color by given downsample pixel index */
-	bool ResetColorDownsampleBufferPixel(const int32 InDownsamplePixelIndex);
-
-	/** Reset the color by given downsample pixel range */
-	bool ResetColorDownsampleBufferPixels(const int32 InDownsamplePixelIndexStart, const int32 InDownsamplePixelIndexEnd);
-
-	/** Remove all pixels from DownsampleBuffer */
-	void EmptyDownsampleBuffer();
-
-private:
-	/** Resize output texture for editor preview */
-	void ResizePreviewRenderTarget(uint32 InSizeX, uint32 InSizeY);
-
-	/** Initialize all textures and creation or loading asset */
-	void Initialize();
-
-	/** Create a render target with unique name */
-	UTextureRenderTarget2D* CreateRenderTarget(const FName& InBaseName);
+	/** Gets the rendered input texture, or nullptr if no input texture is currently rendered */
+	UTexture* GetRenderedInputTexture() const;
 
-public:
 	/**
-	 * Returns the Modulators of the component corresponding to the patch specified. 
-	 * Note, this node does a lookup on all fixture patches in use, hence may be slow and shouldn't be called on tick. 
+	 * Returns the Modulators of the component corresponding to the patch specified.
+	 * Note, this node does a lookup on all fixture patches in use, hence may be slow and shouldn't be called on tick.
 	 */
 	UFUNCTION(BlueprintCallable, Category = "DMX")
 	bool GetPixelMappingComponentModulators(FDMXEntityFixturePatchRef FixturePatchRef, TArray<UDMXModulator*>& DMXModulators);
 
-	/** Returns the preprocess renderer */
-	UDMXPixelMappingPreprocessRenderer* GetPreprocessRenderer() { return PreprocessRenderer; }
-
-#if WITH_EDITOR
-	/** Returns the component canvas used for this widget */
-	FORCEINLINE TSharedPtr<SConstraintCanvas> GetComponentsCanvas() const { return ComponentsCanvas; }
-#endif // WITH_EDITOR
+	/** Returns a copy of the current pixel map render elements */
+	TArray<TSharedRef<UE::DMXPixelMapping::Rendering::FPixelMapRenderElement>> GetPixelMapRenderElements() const;
 
 	/** Type of rendering, Texture, Material, UMG, etc... */
 	UPROPERTY(EditAnywhere, BlueprintReadOnly, Category = "Render Settings")
@@ -180,56 +106,129 @@ public:
 	TObjectPtr<UDMXPixelMappingLayoutScript> LayoutScript;
 
 private:
-	/** Retrieve total count of all output targets that support shared rendering and updates a counter. O(n) */
-	int32 GetTotalDownsamplePixelCount();
+	/** Called when a component was added to or removed from the pixel mapping */
+	void OnComponentAddedOrRemoved(UDMXPixelMapping* PixelMapping, UDMXPixelMappingBaseComponent* Component);
 
-	/** Helper function checks the downsample pixel range */
-	bool IsPixelRangeValid(const int32 InDownsamplePixelIndexStart, const int32 InDownsamplePixelIndexEnd) const;
+	/** Tries to get any world. If with editor, returns the editor world, in game returns GWorld */
+	UWorld* TryGetWorld() const;
 
+	/** Elements rendered during pixel mapping */
+	TArray<TSharedRef<UE::DMXPixelMapping::Rendering::FPixelMapRenderElement>> PixelMapRenderElements;
+
+	/** The user widget instance currently in use */
+	UPROPERTY(Transient, DuplicateTransient)
+	TObjectPtr<UUserWidget> UserWidget;
+
+	/** True once the pixel map was built. Reset to false when the pixel map renderer is invalidated */
+	bool bInvalidatePixelMap = true;
+
+	/** Renderer responsible to preprocess the input texture / material / user widget */
 	UPROPERTY(VisibleAnywhere, Instanced, Category = "Filtering")
 	TObjectPtr<UDMXPixelMappingPreprocessRenderer> PreprocessRenderer;
 
-	/** Reference to renderer */
-	TSharedPtr<IDMXPixelMappingRenderer> PixelMappingRenderer;
+	/** Renderer responsible to pixel map */
+	UPROPERTY()
+	UDMXPixelMappingPixelMapRenderer* PixelMapRenderer;
 
-	/** UMG widget for downsampling */
-	UPROPERTY(Transient)
-	TObjectPtr<UUserWidget> UserWidget;
+
+	//////////////////////
+	// Deprecated Members
+
+public:
+	UE_DEPRECATED(5.3, "Since a pixel mapping asset doesn't have a clear world context, GetWorld() is no longer supported.")
+	UWorld* GetWorld() const;
+
+	UE_DEPRECATED(5.3, "Please use UDMXPixelMappingPixelMapRenderer to render the pixel map")
+	PRAGMA_DISABLE_DEPRECATION_WARNINGS
+	const TSharedPtr<IDMXPixelMappingRenderer>& GetRenderer() { return PixelMappingRenderer_DEPRECATED; }
+	PRAGMA_ENABLE_DEPRECATION_WARNINGS
+
+#if WITH_EDITOR
+	UE_DEPRECATED(5.3, "Please use UDMXPixelMappingPixelMapRenderer to render the pixel map")
+	void RenderEditorPreviewTexture();
+
+	UE_DEPRECATED(5.3, "Please use UDMXPixelMappingPixelMapRenderer to render the pixel map")
+	UTextureRenderTarget2D* GetPreviewRenderTarget();
+#endif
+
+	UE_DEPRECATED(5.3, "Please use UDMXPixelMappingPixelMapRenderer to render the pixel map")
+	FIntPoint GetPixelPosition(int32 InIndex) const;
+
+	UE_DEPRECATED(5.3, "Please use UDMXPixelMappingPixelMapRenderer to render the pixel map")
+	void CreateOrUpdateDownsampleBufferTarget();
+
+	UE_DEPRECATED(5.3, "Please use UDMXPixelMappingPixelMapRenderer to render the pixel map")
+	PRAGMA_DISABLE_DEPRECATION_WARNINGS
+	void AddPixelToDownsampleSet(FDMXPixelMappingDownsamplePixelParamsV2&& InDownsamplePixelParam);
+	PRAGMA_ENABLE_DEPRECATION_WARNINGS
+
+	UE_DEPRECATED(5.3, "Please use UDMXPixelMappingPixelMapRenderer to render the pixel map")
+	int32 GetDownsamplePixelNum();
+
+	UE_DEPRECATED(5.3, "Please use UDMXPixelMappingPixelMapRenderer to render the pixel map")
+	void SetDownsampleBuffer(TArray<FLinearColor>&& InDownsampleBuffer, FIntRect InRect);
+
+	UE_DEPRECATED(5.3, "Please use UDMXPixelMappingPixelMapRenderer to render the pixel map")
+	bool GetDownsampleBufferPixel(const int32 InDownsamplePixelIndex, FLinearColor& OutLinearColor);
+
+	UE_DEPRECATED(5.3, "Please use UDMXPixelMappingPixelMapRenderer to render the pixel map")
+	bool GetDownsampleBufferPixels(const int32 InDownsamplePixelIndexStart, const int32 InDownsamplePixelIndexEnd, TArray<FLinearColor>& OutLinearColors);
+
+	UE_DEPRECATED(5.3, "Please use UDMXPixelMappingPixelMapRenderer to render the pixel map")
+	bool ResetColorDownsampleBufferPixel(const int32 InDownsamplePixelIndex);
+
+	UE_DEPRECATED(5.3, "Please use UDMXPixelMappingPixelMapRenderer to render the pixel map")
+	bool ResetColorDownsampleBufferPixels(const int32 InDownsamplePixelIndexStart, const int32 InDownsamplePixelIndexEnd);
+
+	UE_DEPRECATED(5.3, "Please use UDMXPixelMappingPixelMapRenderer to render the pixel map")
+	void EmptyDownsampleBuffer();
+
+#if WITH_EDITOR
+	UE_DEPRECATED(5.1, "Pixel Mapping Components no longer hold their own widget, in an effort to separate Views from Data.")
+	TSharedRef<SWidget> TakeWidget();
+
+	UE_DEPRECATED(5.1, "Pixel Mapping Components no longer hold their own widget, in an effort to separate Views from Data.")
+	FORCEINLINE TSharedPtr<SConstraintCanvas> GetComponentsCanvas() const { return ComponentsCanvas_DEPRECATED; }
+#endif // WITH_EDITOR
+
+private:
+	UE_DEPRECATED(5.3, "Please use UDMXPixelMappingPixelMapRenderer to render the pixel map")
+	int32 GetTotalDownsamplePixelCount();
+
+	UE_DEPRECATED(5.3, "Please use UDMXPixelMappingPixelMapRenderer to render the pixel map")
+	bool IsPixelRangeValid(const int32 InDownsamplePixelIndexStart, const int32 InDownsamplePixelIndexEnd) const;
+
+	UE_DEPRECATED(5.3, "Please use UDMXPixelMappingPixelMapRenderer to render the pixel map")
+	void ResizePreviewRenderTarget(uint32 InSizeX, uint32 InSizeY);
+
+	UE_DEPRECATED(5.3, "Please use UDMXPixelMappingPixelMapRenderer to render the pixel map")
+	UTextureRenderTarget2D* CreateRenderTarget(const FName& InBaseName);
+
+	PRAGMA_DISABLE_DEPRECATION_WARNINGS
+	TSharedPtr<IDMXPixelMappingRenderer> PixelMappingRenderer_DEPRECATED;
+	PRAGMA_ENABLE_DEPRECATION_WARNINGS
 
 #if WITH_EDITORONLY_DATA
-	/** Canvas for all UI downsamping component widgets */
-	TSharedPtr<SConstraintCanvas> ComponentsCanvas;
+	TSharedPtr<SConstraintCanvas> ComponentsCanvas_DEPRECATED;
 #endif // WITH_EDITORONLY_DATA
 
-	/** GPU downsample pixel buffer target texture */
 	UPROPERTY(Transient, DuplicateTransient)
-	TObjectPtr<UTextureRenderTarget2D> DownsampleBufferTarget;
+	TObjectPtr<UTextureRenderTarget2D> DownsampleBufferTarget_DEPRECATED;
 
 #if WITH_EDITORONLY_DATA
-	/** Editor preview output target */
 	UPROPERTY(Transient, DuplicateTransient)
-	TObjectPtr<UTextureRenderTarget2D> PreviewRenderTarget;
+	TObjectPtr<UTextureRenderTarget2D> PreviewRenderTarget_DEPRECATED;
 #endif // WITH_EDITORONLY_DATA
 
-	/** CPU downsample pixel buffer */
-	TArray<FLinearColor> DownsampleBuffer;
-
-	/** Counter for all pixels from child components */
-	int32 DownsamplePixelCount = 0;
-
-	/** Critical section for set, update and get color array */
-	mutable FCriticalSection DownsampleBufferCS;
-
-	/** Hold the params of the pixels for downsamle rendering */
-	TArray<FDMXPixelMappingDownsamplePixelParamsV2> DownsamplePixelParams;
-
-	/** True once the first render pass completed */
-	bool bWasEverRendered = false;
-
-	/** Initial texture color */
-	static const FLinearColor ClearTextureColor;
+	TArray<FLinearColor> DownsampleBuffer_DEPRECATED;
+	int32 DownsamplePixelCount_DEPRECATED = 0;
+	mutable FCriticalSection DownsampleBufferCS_DEPRECATED;
+	PRAGMA_DISABLE_DEPRECATION_WARNINGS
+	TArray<FDMXPixelMappingDownsamplePixelParamsV2> DownsamplePixelParams_DEPRECATED;
+	PRAGMA_ENABLE_DEPRECATION_WARNINGS
+	static const FLinearColor ClearTextureColor_DEPRECATED;
+	bool bWasEverRendered_DEPRECATED = false;
 
 public:
-	/** Max downsample target size */
-	static const FIntPoint MaxDownsampleBufferTargetSize;
+	static const FIntPoint MaxDownsampleBufferTargetSize_DEPRECATED;
 };
diff --git a/Engine/Plugins/VirtualProduction/DMX/DMXProtocol/Source/DMXProtocol/Public/IO/DMXOutputPort.h b/Engine/Plugins/VirtualProduction/DMX/DMXProtocol/Source/DMXProtocol/Public/IO/DMXOutputPort.h
--- a/Engine/Plugins/VirtualProduction/DMX/DMXProtocol/Source/DMXProtocol/Public/IO/DMXOutputPort.h
+++ b/Engine/Plugins/VirtualProduction/DMX/DMXProtocol/Source/DMXProtocol/Public/IO/DMXOutputPort.h
@@ -5,14 +5,12 @@
 #include "DMXPort.h"
 
 #include "DMXProtocolCommon.h"
-
-#include "CoreMinimal.h"
-#include "Tickable.h"
 #include "Containers/Queue.h" 
 #include "HAL/Runnable.h"
 #include "Misc/ScopeLock.h" 
 #include "Misc/SingleThreadRunnable.h"
 #include "Templates/Atomic.h"
+#include "Tickable.h"
 
 struct FDMXOutputPortConfig;
 struct FDMXOutputPortDestinationAddress;
