DMX - Fix Pixel Mapping DMX Library tab doesn't update when patches were removed from pixel mapping

diff --git a/Engine/Plugins/VirtualProduction/DMX/DMXPixelMapping/Source/DMXPixelMappingEditor/Private/Views/SDMXPixelMappingDMXLibraryView.cpp b/Engine/Plugins/VirtualProduction/DMX/DMXPixelMapping/Source/DMXPixelMappingEditor/Private/Views/SDMXPixelMappingDMXLibraryView.cpp
--- a/Engine/Plugins/VirtualProduction/DMX/DMXPixelMapping/Source/DMXPixelMappingEditor/Private/Views/SDMXPixelMappingDMXLibraryView.cpp
+++ b/Engine/Plugins/VirtualProduction/DMX/DMXPixelMapping/Source/DMXPixelMappingEditor/Private/Views/SDMXPixelMappingDMXLibraryView.cpp
@@ -40,6 +40,12 @@ void SDMXPixelMappingDMXLibraryView::Construct(const FArguments& InArgs, const T
 	}
 	const TSharedRef<FDMXPixelMappingToolkit> Toolkit = WeakToolkit.Pin().ToSharedRef();
 
+	const TAttribute<bool> AddDMXLibraryButtonEnabledAttribute = TAttribute<bool>::CreateLambda(
+		[this]()
+		{
+			return bRenderComponentContainedInSelection;
+		});
+
 	ChildSlot
 	[
 		SNew(SBorder)
@@ -56,9 +62,15 @@ void SDMXPixelMappingDMXLibraryView::Construct(const FArguments& InArgs, const T
 				SNew(SButton)
 				.ButtonStyle(FAppStyle::Get(), "FlatButton.Success")
 				.ForegroundColor(FLinearColor::White)
-				.ToolTipText(LOCTEXT("AddFixtureGroupTooltip", "Adds a Fixture Group to the Pixel Mapping"))
+				.ToolTipText_Lambda([this]()
+					{
+						return bRenderComponentContainedInSelection ?
+							LOCTEXT("AddFixtureGroupTooltip", "Adds a DMX Library to the Pixel Mapping") :
+							LOCTEXT("CannotAddFixtureGroupTooltip", "Please select an Input Source");
+					})
 				.ContentPadding(FMargin(5.0f, 1.0f))
 				.OnClicked(this, &SDMXPixelMappingDMXLibraryView::OnAddDMXLibraryButtonClicked)
+				.IsEnabled(AddDMXLibraryButtonEnabledAttribute)
 				.Content()
 				[
 					SNew(SHorizontalBox)
@@ -96,8 +108,13 @@ void SDMXPixelMappingDMXLibraryView::Construct(const FArguments& InArgs, const T
 
 	UDMXLibrary::GetOnEntitiesAdded().AddSP(this, &SDMXPixelMappingDMXLibraryView::OnEntitiesAddedOrRemoved);
 	UDMXLibrary::GetOnEntitiesRemoved().AddSP(this, &SDMXPixelMappingDMXLibraryView::OnEntitiesAddedOrRemoved);
+	UDMXPixelMappingBaseComponent::GetOnComponentAdded().AddSP(this, &SDMXPixelMappingDMXLibraryView::OnComponentAddedOrRemoved);
+	UDMXPixelMappingBaseComponent::GetOnComponentRemoved().AddSP(this, &SDMXPixelMappingDMXLibraryView::OnComponentAddedOrRemoved);
 
 	Toolkit->GetOnSelectedComponentsChangedDelegate().AddSP(this, &SDMXPixelMappingDMXLibraryView::OnComponentSelected);
+
+	// Apply the current selection
+	OnComponentSelected();
 }
 
 void SDMXPixelMappingDMXLibraryView::RequestRefresh()
@@ -113,6 +130,26 @@ void SDMXPixelMappingDMXLibraryView::AddReferencedObjects(FReferenceCollector& C
 	Collector.AddReferencedObjects(ViewModels);
 }
 
+void SDMXPixelMappingDMXLibraryView::PostUndo(bool bSuccess)
+{
+	RequestRefresh();
+}
+
+void SDMXPixelMappingDMXLibraryView::PostRedo(bool bSuccess)
+{
+	RequestRefresh();
+}
+
+void SDMXPixelMappingDMXLibraryView::OnComponentAddedOrRemoved(UDMXPixelMapping* PixelMapping, UDMXPixelMappingBaseComponent* Component)
+{
+	RequestRefresh();
+}
+
+void SDMXPixelMappingDMXLibraryView::OnEntitiesAddedOrRemoved(UDMXLibrary* DMXLibrary, TArray<UDMXEntity*> Entities)
+{
+	RequestRefresh();
+}
+
 void SDMXPixelMappingDMXLibraryView::ForceRefresh()
 {
 	RefreshTimerHandle.Invalidate();
@@ -191,11 +228,6 @@ void SDMXPixelMappingDMXLibraryView::ForceRefresh()
 	}
 }
 
-void SDMXPixelMappingDMXLibraryView::OnEntitiesAddedOrRemoved(UDMXLibrary* DMXLibrary, TArray<UDMXEntity*> Entities)
-{
-	RequestRefresh();
-}
-
 void SDMXPixelMappingDMXLibraryView::OnComponentSelected()
 {
 	const TSharedPtr<FDMXPixelMappingToolkit> Toolkit = WeakToolkit.Pin();
@@ -204,28 +236,35 @@ void SDMXPixelMappingDMXLibraryView::OnComponentSelected()
 		return;
 	}
 
-	// Refresh the view if there is no view model for the fixture group of a selected component
-	const TSet<FDMXPixelMappingComponentReference>& SelectedComponents = Toolkit->GetSelectedComponents();
+	bRenderComponentContainedInSelection = false;
+	bool bHasValidModel = false;
 
+	const TSet<FDMXPixelMappingComponentReference>& SelectedComponents = Toolkit->GetSelectedComponents();
 	UDMXPixelMappingFixtureGroupComponent* ComponentToFocus = nullptr;
 	for (const FDMXPixelMappingComponentReference& ComponentReference : SelectedComponents)
 	{
 		UDMXPixelMappingBaseComponent* Component = ComponentReference.GetComponent();
 		do
 		{
-			TObjectPtr<UDMXPixelMappingDMXLibraryViewModel>* CorrespondingViewModelPtr = Algo::FindByPredicate(ViewModels, [Component](const UDMXPixelMappingDMXLibraryViewModel* ViewModel)
+			bRenderComponentContainedInSelection |= Component->GetClass() == UDMXPixelMappingRendererComponent::StaticClass();
+
+			const TObjectPtr<UDMXPixelMappingDMXLibraryViewModel>* const CorrespondingViewModelPtr = Algo::FindByPredicate(ViewModels, [Component](const UDMXPixelMappingDMXLibraryViewModel* ViewModel)
 				{
 					return ViewModel->GetFixtureGroupComponent() == Component;
 				});
-			if (!CorrespondingViewModelPtr)
+			if (CorrespondingViewModelPtr)
 			{
-				RequestRefresh();
-				return;
+				bHasValidModel = true;
 			}
 
 			Component = Component->GetParent();
 		} while (Component);
 	}
+
+	if (!bHasValidModel)
+	{
+		RequestRefresh();
+	}
 }
 
 FReply SDMXPixelMappingDMXLibraryView::OnAddDMXLibraryButtonClicked()
diff --git a/Engine/Plugins/VirtualProduction/DMX/DMXPixelMapping/Source/DMXPixelMappingEditor/Private/Views/SDMXPixelMappingDMXLibraryView.h b/Engine/Plugins/VirtualProduction/DMX/DMXPixelMapping/Source/DMXPixelMappingEditor/Private/Views/SDMXPixelMappingDMXLibraryView.h
--- a/Engine/Plugins/VirtualProduction/DMX/DMXPixelMapping/Source/DMXPixelMappingEditor/Private/Views/SDMXPixelMappingDMXLibraryView.h
+++ b/Engine/Plugins/VirtualProduction/DMX/DMXPixelMapping/Source/DMXPixelMappingEditor/Private/Views/SDMXPixelMappingDMXLibraryView.h
@@ -4,12 +4,15 @@
 
 #include "Library/DMXEntityReference.h"
 
+#include "EditorUndoClient.h"
 #include "Engine/EngineTypes.h"
 #include "Widgets/SCompoundWidget.h"
 #include "UObject/GCObject.h"
 
 class FDMXPixelMappingToolkit;
 class SDMXReadOnlyFixturePatchList;
+class UDMXPixelMapping;
+class UDMXPixelMappingBaseComponent;
 class UDMXPixelMappingFixtureGroupComponent;
 class UDMXPixelMappingDMXLibraryViewModel;
 class UDMXPixelMappingRootComponent;
@@ -23,6 +26,7 @@ class SScrollBox;
 class SDMXPixelMappingDMXLibraryView
 	: public SCompoundWidget
 	, public FGCObject
+	, public FSelfRegisteringEditorUndoClient
 {
 public:
 	SLATE_BEGIN_ARGS(SDMXPixelMappingDMXLibraryView) { }
@@ -43,9 +47,14 @@ protected:
 	}
 	//~ End FGCObject interface
 
+	//~ Begin FSelfRegisteringEditorUndoClient interface
+	virtual void PostUndo(bool bSuccess) override;
+	virtual void PostRedo(bool bSuccess) override;
+	//~ End FSelfRegisteringEditorUndoClient interface
+
 private:
-	/** Refreshes the widget */
-	void ForceRefresh();
+	/** Called when a component was added to or removed from the pixel mapping */
+	void OnComponentAddedOrRemoved(UDMXPixelMapping* PixelMapping, UDMXPixelMappingBaseComponent* Component);
 
 	/** Called when entities were added to or removed from a DMX Library */
 	void OnEntitiesAddedOrRemoved(UDMXLibrary* DMXLibrary, TArray<UDMXEntity*> Entities);
@@ -53,6 +62,9 @@ private:
 	/** Called when a component was selected */
 	void OnComponentSelected();
 
+	/** Refreshes the widget */
+	void ForceRefresh();
+
 	/** Called when the 'Add DMX Library' button was clicked */
 	FReply OnAddDMXLibraryButtonClicked();
 
@@ -77,6 +89,9 @@ private:
 	/** Timer handle for the Request Refresh method */
 	FTimerHandle RefreshTimerHandle;
 
+	/** True if a renderere component is contained in the current selection */
+	bool bRenderComponentContainedInSelection = false;
+
 	/** The view model of the focused DMX Library */
 	TArray<TObjectPtr<UDMXPixelMappingDMXLibraryViewModel>> ViewModels;
 
diff --git a/Engine/Plugins/VirtualProduction/DMX/DMXPixelMapping/Source/DMXPixelMappingEditor/Private/Views/SDMXPixelMappingHierarchyView.cpp b/Engine/Plugins/VirtualProduction/DMX/DMXPixelMapping/Source/DMXPixelMappingEditor/Private/Views/SDMXPixelMappingHierarchyView.cpp
--- a/Engine/Plugins/VirtualProduction/DMX/DMXPixelMapping/Source/DMXPixelMappingEditor/Private/Views/SDMXPixelMappingHierarchyView.cpp
+++ b/Engine/Plugins/VirtualProduction/DMX/DMXPixelMapping/Source/DMXPixelMappingEditor/Private/Views/SDMXPixelMappingHierarchyView.cpp
@@ -58,7 +58,7 @@ public:
 
 void SDMXPixelMappingHierarchyView::Construct(const FArguments& InArgs, const TSharedPtr<FDMXPixelMappingToolkit>& InToolkit)
 {
-	Toolkit = InToolkit;
+	WeakToolkit = InToolkit;
 
 	bRebuildTreeRequested = false;
 	bIsUpdatingSelection = false;
@@ -170,9 +170,14 @@ FReply SDMXPixelMappingHierarchyView::OnKeyDown(const FGeometry& MyGeometry, con
 
 TSharedPtr<SWidget> SDMXPixelMappingHierarchyView::WidgetHierarchy_OnContextMenuOpening()
 {
+	if (!WeakToolkit.IsValid())
+	{
+		return SNullWidget::NullWidget;
+	}
+
 	FMenuBuilder MenuBuilder(true, CommandList);
 
-	FDMXPixelMappingEditorUtils::CreateComponentContextMenu(MenuBuilder, Toolkit.Pin().ToSharedRef());
+	FDMXPixelMappingEditorUtils::CreateComponentContextMenu(MenuBuilder, WeakToolkit.Pin().ToSharedRef());
 
 	return MenuBuilder.MakeWidget();
 }
@@ -193,7 +198,7 @@ void SDMXPixelMappingHierarchyView::WidgetHierarchy_OnSelectionChanged(FDMXPixel
 	{
 		bIsUpdatingSelection = true;
 
-		if (TSharedPtr<FDMXPixelMappingToolkit> ToolkitPtr = Toolkit.Pin())
+		if (TSharedPtr<FDMXPixelMappingToolkit> Toolkit = WeakToolkit.Pin())
 		{
 			TSet<FDMXPixelMappingComponentReference> ComponentsToSelect;
 			FDMXPixelMappingHierarchyItemWidgetModelArr SelectedItems = WidgetTreeView->GetSelectedItems();
@@ -202,7 +207,7 @@ void SDMXPixelMappingHierarchyView::WidgetHierarchy_OnSelectionChanged(FDMXPixel
 				ComponentsToSelect.Add(Item->GetReference());
 			}
 
-			ToolkitPtr->SelectComponents(ComponentsToSelect);
+			Toolkit->SelectComponents(ComponentsToSelect);
 		}
 
 		bIsUpdatingSelection = false;
@@ -367,8 +372,8 @@ bool SDMXPixelMappingHierarchyView::RestoreSelectionForItemAndChildren(FDMXPixel
 		return false;
 	}
 
-	const TSharedPtr<FDMXPixelMappingToolkit> StrongToolkit = Toolkit.Pin();
-	if (!StrongToolkit.IsValid())
+	const TSharedPtr<FDMXPixelMappingToolkit> Toolkit = WeakToolkit.Pin();
+	if (!Toolkit.IsValid())
 	{
 		return false;
 	}
@@ -382,7 +387,7 @@ bool SDMXPixelMappingHierarchyView::RestoreSelectionForItemAndChildren(FDMXPixel
 		bContainsSelection |= RestoreSelectionForItemAndChildren(ChildModel);
 	}
 
-	const TSet<FDMXPixelMappingComponentReference> SelectedComponents = StrongToolkit->GetSelectedComponents();
+	const TSet<FDMXPixelMappingComponentReference> SelectedComponents = Toolkit->GetSelectedComponents();
 	if (bContainsSelection)
 	{
 		WidgetTreeView->SetItemExpansion(Model, true);
@@ -401,8 +406,13 @@ bool SDMXPixelMappingHierarchyView::RestoreSelectionForItemAndChildren(FDMXPixel
 
 void SDMXPixelMappingHierarchyView::RefreshTree()
 {
+	if (!WeakToolkit.IsValid())
+	{
+		return;
+	}
+
 	RootWidgets.Empty();
-	RootWidgets.Add(MakeShared<FDMXPixelMappingHierarchyItemWidgetModel>(Toolkit.Pin()));
+	RootWidgets.Add(MakeShared<FDMXPixelMappingHierarchyItemWidgetModel>(WeakToolkit.Pin()));
 
 	FilterHandler->RefreshAndFilterTree();
 }
@@ -619,44 +629,22 @@ void SDMXPixelMappingHierarchyView::RequestComponentRedraw(UDMXPixelMappingBaseC
 		return;
 	}
 
-	if (TSharedPtr<FDMXPixelMappingToolkit> ToolkitPtr = Toolkit.Pin())
+	if (TSharedPtr<FDMXPixelMappingToolkit> Toolkit = WeakToolkit.Pin())
 	{
 		TSet<FDMXPixelMappingComponentReference> ComponentsToUpdate;
-		FDMXPixelMappingComponentReference ComponentReference(ToolkitPtr, Component);
+		FDMXPixelMappingComponentReference ComponentReference(Toolkit, Component);
 		ComponentsToUpdate.Add(ComponentReference);
-		//ToolkitPtr->SelectComponents(ComponentsToUpdate);
 	}
 }
 
 void SDMXPixelMappingHierarchyView::BeginDelete()
 {
-	FDMXPixelMappingHierarchyItemWidgetModelArr SelectedItems = WidgetTreeView->GetSelectedItems();
-
-	const FScopedTransaction Transaction(FText::Format(LOCTEXT("DMXPixelMapping.RemoveComponents", "Remove {0}|plural(one=Component, other=Components)"), SelectedItems.Num()));
-
-	if (TSharedPtr<FDMXPixelMappingToolkit> ToolkitPtr = Toolkit.Pin())
+	if (TSharedPtr<FDMXPixelMappingToolkit> Toolkit = WeakToolkit.Pin())
 	{
-		for (FDMXPixelMappingHierarchyItemWidgetModelPtr SelectedItem : SelectedItems)
-		{
-			UDMXPixelMappingBaseComponent* SelectedComponent = SelectedItem->GetReference().GetComponent();
-			if (SelectedComponent)
-			{
-				constexpr bool bModifyChildrenRecursively = true;
-				SelectedComponent->ForEachChild([](UDMXPixelMappingBaseComponent* ChildComponent)
-					{
-						ChildComponent->Modify();
-					}, bModifyChildrenRecursively);
-
-				UDMXPixelMappingBaseComponent* ParentComponent = SelectedComponent->GetParent();
-				if (ParentComponent)
-				{
-					ParentComponent->Modify();
-					SelectedComponent->Modify();
-
-					ParentComponent->RemoveChild(SelectedComponent);
-				}
-			}
-		}
+		const int32 NumSelectedComponents = Toolkit->GetSelectedComponents().Num();
+		const FScopedTransaction Transaction(FText::Format(LOCTEXT("DMXPixelMapping.RemoveComponents", "Remove {0}|plural(one=Component, other=Components)"), NumSelectedComponents));
+		
+		Toolkit->DeleteSelectedComponents();
 	}
 }
 
diff --git a/Engine/Plugins/VirtualProduction/DMX/DMXPixelMapping/Source/DMXPixelMappingEditor/Private/Views/SDMXPixelMappingHierarchyView.h b/Engine/Plugins/VirtualProduction/DMX/DMXPixelMapping/Source/DMXPixelMappingEditor/Private/Views/SDMXPixelMappingHierarchyView.h
--- a/Engine/Plugins/VirtualProduction/DMX/DMXPixelMapping/Source/DMXPixelMappingEditor/Private/Views/SDMXPixelMappingHierarchyView.h
+++ b/Engine/Plugins/VirtualProduction/DMX/DMXPixelMapping/Source/DMXPixelMappingEditor/Private/Views/SDMXPixelMappingHierarchyView.h
@@ -127,7 +127,7 @@ private:
 	void SelectFirstAvailableRenderer();
 
 private:
-	TWeakPtr<FDMXPixelMappingToolkit> Toolkit;
+	TWeakPtr<FDMXPixelMappingToolkit> WeakToolkit;
 
 	/** Commands specific to the hierarchy. */
 	TSharedPtr<FUICommandList> CommandList;
