DMX - Fix CIS Check Balanced Macros

diff --git a/Engine/Plugins/VirtualProduction/DMX/DMXPixelMapping/Source/DMXPixelMappingRuntime/Private/Components/DMXPixelMappingMatrixCellComponent.cpp b/Engine/Plugins/VirtualProduction/DMX/DMXPixelMapping/Source/DMXPixelMappingRuntime/Private/Components/DMXPixelMappingMatrixCellComponent.cpp
--- a/Engine/Plugins/VirtualProduction/DMX/DMXPixelMapping/Source/DMXPixelMappingRuntime/Private/Components/DMXPixelMappingMatrixCellComponent.cpp
+++ b/Engine/Plugins/VirtualProduction/DMX/DMXPixelMapping/Source/DMXPixelMappingRuntime/Private/Components/DMXPixelMappingMatrixCellComponent.cpp
@@ -285,9 +285,7 @@ void UDMXPixelMappingMatrixCellComponent::QueueDownsample()
 	}
 
 	// Store downsample index
-	PRAGMA_DISABLE_DEPRECATION_WARNINGS
 	DownsamplePixelIndex = RendererComponent->GetDownsamplePixelNum();
-	PRAGMA_ENABLE_DEPRECATION_WARNINGS
 
 	const uint32 TextureSizeX = InputTexture->GetResource()->GetSizeX();
 	const uint32 TextureSizeY = InputTexture->GetResource()->GetSizeY();
