DMX - Fix pixel mapping renderer component doesn't follow texture size

diff --git a/Engine/Plugins/VirtualProduction/DMX/DMXPixelMapping/Source/DMXPixelMappingRuntime/Private/Components/DMXPixelMappingFixtureGroupItemComponent.cpp b/Engine/Plugins/VirtualProduction/DMX/DMXPixelMapping/Source/DMXPixelMappingRuntime/Private/Components/DMXPixelMappingFixtureGroupItemComponent.cpp
--- a/Engine/Plugins/VirtualProduction/DMX/DMXPixelMapping/Source/DMXPixelMappingRuntime/Private/Components/DMXPixelMappingFixtureGroupItemComponent.cpp
+++ b/Engine/Plugins/VirtualProduction/DMX/DMXPixelMapping/Source/DMXPixelMappingRuntime/Private/Components/DMXPixelMappingFixtureGroupItemComponent.cpp
@@ -461,8 +461,8 @@ void UDMXPixelMappingFixtureGroupItemComponent::UpdateRenderElement()
 
 	const UDMXPixelMappingRendererComponent* RendererComponent = GetRendererComponent();
 	const UTexture* InputTexture = RendererComponent ? RendererComponent->GetRenderedInputTexture() : nullptr;
-	const double InputTextureWidth = InputTexture ? InputTexture->GetSurfaceWidth() : 0.0;
-	const double InputTextureHeight = InputTexture ? InputTexture->GetSurfaceHeight() : 0.0;
+	const double InputTextureWidth = InputTexture ? InputTexture->GetSurfaceWidth() : 1.0;
+	const double InputTextureHeight = InputTexture ? InputTexture->GetSurfaceHeight() : 1.0;
 
 	FPixelMapRenderElementParameters Parameters;
 	Parameters.UV = FVector2D(GetPosition().X / InputTextureWidth, GetPosition().Y / InputTextureHeight);
diff --git a/Engine/Plugins/VirtualProduction/DMX/DMXPixelMapping/Source/DMXPixelMappingRuntime/Private/Components/DMXPixelMappingMatrixCellComponent.cpp b/Engine/Plugins/VirtualProduction/DMX/DMXPixelMapping/Source/DMXPixelMappingRuntime/Private/Components/DMXPixelMappingMatrixCellComponent.cpp
--- a/Engine/Plugins/VirtualProduction/DMX/DMXPixelMapping/Source/DMXPixelMappingRuntime/Private/Components/DMXPixelMappingMatrixCellComponent.cpp
+++ b/Engine/Plugins/VirtualProduction/DMX/DMXPixelMapping/Source/DMXPixelMappingRuntime/Private/Components/DMXPixelMappingMatrixCellComponent.cpp
@@ -336,11 +336,11 @@ TSharedRef<UE::DMXPixelMapping::Rendering::FPixelMapRenderElement> UDMXPixelMapp
 void UDMXPixelMappingMatrixCellComponent::UpdateRenderElement()
 {
 	using namespace UE::DMXPixelMapping::Rendering;
-
+	
 	const UDMXPixelMappingRendererComponent* RendererComponent = GetRendererComponent();
 	const UTexture* InputTexture = RendererComponent ? RendererComponent->GetRenderedInputTexture() : nullptr;
-	const double InputTextureWidth = InputTexture ? InputTexture->GetSurfaceWidth() : 0.0;
-	const double InputTextureHeight = InputTexture ? InputTexture->GetSurfaceHeight() : 0.0;
+	const double InputTextureWidth = InputTexture ? InputTexture->GetSurfaceWidth() : 1.0;
+	const double InputTextureHeight = InputTexture ? InputTexture->GetSurfaceHeight() : 1.0;
 
 	FPixelMapRenderElementParameters Parameters;
 	Parameters.UV = FVector2D(GetPosition().X / InputTextureWidth, GetPosition().Y / InputTextureHeight);
diff --git a/Engine/Plugins/VirtualProduction/DMX/DMXPixelMapping/Source/DMXPixelMappingRuntime/Private/Components/DMXPixelMappingRendererComponent.cpp b/Engine/Plugins/VirtualProduction/DMX/DMXPixelMapping/Source/DMXPixelMappingRuntime/Private/Components/DMXPixelMappingRendererComponent.cpp
--- a/Engine/Plugins/VirtualProduction/DMX/DMXPixelMapping/Source/DMXPixelMappingRuntime/Private/Components/DMXPixelMappingRendererComponent.cpp
+++ b/Engine/Plugins/VirtualProduction/DMX/DMXPixelMapping/Source/DMXPixelMappingRuntime/Private/Components/DMXPixelMappingRendererComponent.cpp
@@ -180,8 +180,6 @@ void UDMXPixelMappingRendererComponent::InvalidatePreprocessRenderer()
 	default:
 		checkf(0, TEXT("Invalid Renderer Type in DMXPixelMappingRendererComponent"));
 	}
-
-	SetSize(PreprocessRenderer->GetResultingSize2D());
 }
 
 void UDMXPixelMappingRendererComponent::InvalidatePixelMapRenderer()
@@ -253,6 +251,17 @@ void UDMXPixelMappingRendererComponent::Render()
 {
 	SCOPE_CYCLE_COUNTER(STAT_DMXPixelMappingRender);
 
+	// Always size to texture
+	if (UTexture* Texture = PreprocessRenderer->GetRenderedTexture())
+	{
+		const FVector2D TextureSize(Texture->GetSurfaceWidth(), Texture->GetSurfaceHeight());
+		if (GetSize() != TextureSize)
+		{
+			SetSize(TextureSize);
+			bInvalidatePixelMap = true;
+		}
+	}
+
 	PreprocessRenderer->Render();
 
 	UTexture* RenderedInputTexture = PreprocessRenderer->GetRenderedTexture();
