DMX - Fix Pixel Mapping Matrix Component sending incorrect values, fix not all component changes are reflected in the pixel map and preview

diff --git a/Engine/Plugins/VirtualProduction/DMX/DMXPixelMapping/Source/DMXPixelMappingEditor/Private/DragDrop/DMXPixelMappingDragDropOp.cpp b/Engine/Plugins/VirtualProduction/DMX/DMXPixelMapping/Source/DMXPixelMappingEditor/Private/DragDrop/DMXPixelMappingDragDropOp.cpp
--- a/Engine/Plugins/VirtualProduction/DMX/DMXPixelMapping/Source/DMXPixelMappingEditor/Private/DragDrop/DMXPixelMappingDragDropOp.cpp
+++ b/Engine/Plugins/VirtualProduction/DMX/DMXPixelMapping/Source/DMXPixelMappingEditor/Private/DragDrop/DMXPixelMappingDragDropOp.cpp
@@ -87,11 +87,6 @@ void FDMXPixelMappingDragDropOp::LayoutOutputComponents(const FVector2D& GraphSp
 				{
 					OutputComponent->PreEditChange(nullptr);
 
-					constexpr bool bNotifyChildrenRecursive = true;
-					Component->ForEachChild([](UDMXPixelMappingBaseComponent* Component)
-						{
-							Component->PreEditChange(nullptr);
-						}, bNotifyChildrenRecursive);
 
 					if (ensureMsgf(OutputComponent->GetClass() != UDMXPixelMappingMatrixComponent::StaticClass(),
 						TEXT("Matrix Cells are not supported. Use the GroupChildDragDropHelper from this class instead")))
@@ -101,14 +96,9 @@ void FDMXPixelMappingDragDropOp::LayoutOutputComponents(const FVector2D& GraphSp
 						const FVector2D NewPosition = GraphSpacePosition - AnchorOffset - GraphSpaceDragOffset;
 						OutputComponent->SetPosition(NewPosition.RoundToVector());
 
-						OutputComponent->PostEditChange();
 					}
 
-					Component->ForEachChild([](UDMXPixelMappingBaseComponent* Component)
-						{
-							Component->PostEditChange();
-						}, bNotifyChildrenRecursive);
-
+					OutputComponent->PostEditChange();
 				}
 			}
 		}
diff --git a/Engine/Plugins/VirtualProduction/DMX/DMXPixelMapping/Source/DMXPixelMappingEditor/Private/DragDrop/DMXPixelMappingGroupChildDragDropHelper.cpp b/Engine/Plugins/VirtualProduction/DMX/DMXPixelMapping/Source/DMXPixelMappingEditor/Private/DragDrop/DMXPixelMappingGroupChildDragDropHelper.cpp
--- a/Engine/Plugins/VirtualProduction/DMX/DMXPixelMapping/Source/DMXPixelMappingEditor/Private/DragDrop/DMXPixelMappingGroupChildDragDropHelper.cpp
+++ b/Engine/Plugins/VirtualProduction/DMX/DMXPixelMapping/Source/DMXPixelMappingEditor/Private/DragDrop/DMXPixelMappingGroupChildDragDropHelper.cpp
@@ -91,6 +91,8 @@ void FDMXPixelMappingGroupChildDragDropHelper::LayoutAligned(const FVector2D& Gr
 			{
 				if (UDMXPixelMappingOutputComponent* ChildComponent = WeakChildComponent.Get())
 				{
+					ChildComponent->PreEditChange(nullptr);
+
 					if (GroupComponent->IsOverPosition(NextPosition) && 
 						GroupComponent->IsOverPosition(NextPosition + ChildComponent->GetSize()))
 					{
@@ -120,6 +122,8 @@ void FDMXPixelMappingGroupChildDragDropHelper::LayoutAligned(const FVector2D& Gr
 							NextPosition = FVector2D(NextPosition.X + ChildComponent->GetSize().X, NextPosition.Y);
 						}
 					}
+
+					ChildComponent->PostEditChange();
 				}
 			}
 		}
@@ -141,10 +145,14 @@ void FDMXPixelMappingGroupChildDragDropHelper::LayoutUnaligned(const FVector2D&
 				{
 					if (UDMXPixelMappingOutputComponent* ChildComponent = WeakOutputComponent.Get())
 					{
+						ChildComponent->PreEditChange(nullptr);
+
 						FVector2D AnchorOffset = Anchor - ChildComponent->GetPosition();
 
 						FVector2D NewPosition = GraphSpacePosition - AnchorOffset - PinnedDragDropOp->GraphSpaceDragOffset;
 						SetPosition(ChildComponent, NewPosition);
+
+						ChildComponent->PostEditChange();
 					}
 				}
 			}
diff --git a/Engine/Plugins/VirtualProduction/DMX/DMXPixelMapping/Source/DMXPixelMappingEditor/Private/ViewModels/DMXPixelMappingOutputComponentModel.cpp b/Engine/Plugins/VirtualProduction/DMX/DMXPixelMapping/Source/DMXPixelMappingEditor/Private/ViewModels/DMXPixelMappingOutputComponentModel.cpp
--- a/Engine/Plugins/VirtualProduction/DMX/DMXPixelMapping/Source/DMXPixelMappingEditor/Private/ViewModels/DMXPixelMappingOutputComponentModel.cpp
+++ b/Engine/Plugins/VirtualProduction/DMX/DMXPixelMapping/Source/DMXPixelMappingEditor/Private/ViewModels/DMXPixelMappingOutputComponentModel.cpp
@@ -110,6 +110,7 @@ bool FDMXPixelMappingOutputComponentModel::ShouldDrawCellID() const
 
 	return
 		LayoutSettings->bShowCellIDs &&
+		WeakOutputComponent.IsValid() &&
 		WeakOutputComponent->GetClass() == UDMXPixelMappingMatrixCellComponent::StaticClass();
 }
 
diff --git a/Engine/Plugins/VirtualProduction/DMX/DMXPixelMapping/Source/DMXPixelMappingRuntime/Private/Components/DMXPixelMappingFixtureGroupComponent.cpp b/Engine/Plugins/VirtualProduction/DMX/DMXPixelMapping/Source/DMXPixelMappingRuntime/Private/Components/DMXPixelMappingFixtureGroupComponent.cpp
--- a/Engine/Plugins/VirtualProduction/DMX/DMXPixelMapping/Source/DMXPixelMappingRuntime/Private/Components/DMXPixelMappingFixtureGroupComponent.cpp
+++ b/Engine/Plugins/VirtualProduction/DMX/DMXPixelMapping/Source/DMXPixelMappingRuntime/Private/Components/DMXPixelMappingFixtureGroupComponent.cpp
@@ -66,6 +66,8 @@ void UDMXPixelMappingFixtureGroupComponent::PostEditChangeProperty(FPropertyChan
 		}
 	}
 	PRAGMA_ENABLE_DEPRECATION_WARNINGS
+
+	InvalidatePixelMapRenderer();
 }
 #endif // WITH_EDITOR
 
diff --git a/Engine/Plugins/VirtualProduction/DMX/DMXPixelMapping/Source/DMXPixelMappingRuntime/Private/Components/DMXPixelMappingFixtureGroupItemComponent.cpp b/Engine/Plugins/VirtualProduction/DMX/DMXPixelMapping/Source/DMXPixelMappingRuntime/Private/Components/DMXPixelMappingFixtureGroupItemComponent.cpp
--- a/Engine/Plugins/VirtualProduction/DMX/DMXPixelMapping/Source/DMXPixelMappingRuntime/Private/Components/DMXPixelMappingFixtureGroupItemComponent.cpp
+++ b/Engine/Plugins/VirtualProduction/DMX/DMXPixelMapping/Source/DMXPixelMappingRuntime/Private/Components/DMXPixelMappingFixtureGroupItemComponent.cpp
@@ -87,10 +87,7 @@ void UDMXPixelMappingFixtureGroupItemComponent::Serialize(FArchive& Ar)
 	}
 #endif
 
-	if (UDMXPixelMappingRendererComponent* RendererComponent = GetRendererComponent())
-	{
-		RendererComponent->InvalidatePixelMapRenderer();
-	}
+	InvalidatePixelMapRenderer();
 }
 
 void UDMXPixelMappingFixtureGroupItemComponent::PostInitProperties()
@@ -201,6 +198,8 @@ void UDMXPixelMappingFixtureGroupItemComponent::PostEditChangeProperty(FProperty
 		}
 	}
 	PRAGMA_ENABLE_DEPRECATION_WARNINGS
+
+	InvalidatePixelMapRenderer();
 }
 #endif // WITH_EDITOR
 
diff --git a/Engine/Plugins/VirtualProduction/DMX/DMXPixelMapping/Source/DMXPixelMappingRuntime/Private/Components/DMXPixelMappingMatrixCellComponent.cpp b/Engine/Plugins/VirtualProduction/DMX/DMXPixelMapping/Source/DMXPixelMappingRuntime/Private/Components/DMXPixelMappingMatrixCellComponent.cpp
--- a/Engine/Plugins/VirtualProduction/DMX/DMXPixelMapping/Source/DMXPixelMappingRuntime/Private/Components/DMXPixelMappingMatrixCellComponent.cpp
+++ b/Engine/Plugins/VirtualProduction/DMX/DMXPixelMapping/Source/DMXPixelMappingRuntime/Private/Components/DMXPixelMappingMatrixCellComponent.cpp
@@ -102,10 +102,7 @@ void UDMXPixelMappingMatrixCellComponent::PostEditChangeProperty(FPropertyChange
 	}
 	PRAGMA_ENABLE_DEPRECATION_WARNINGS
 
-	if (UDMXPixelMappingRendererComponent* RendererComponent = GetRendererComponent())
-	{
-		RendererComponent->InvalidatePixelMapRenderer();
-	}
+	InvalidatePixelMapRenderer();
 }
 #endif // WITH_EDITOR
 
diff --git a/Engine/Plugins/VirtualProduction/DMX/DMXPixelMapping/Source/DMXPixelMappingRuntime/Private/Components/DMXPixelMappingMatrixComponent.cpp b/Engine/Plugins/VirtualProduction/DMX/DMXPixelMapping/Source/DMXPixelMappingRuntime/Private/Components/DMXPixelMappingMatrixComponent.cpp
--- a/Engine/Plugins/VirtualProduction/DMX/DMXPixelMapping/Source/DMXPixelMappingRuntime/Private/Components/DMXPixelMappingMatrixComponent.cpp
+++ b/Engine/Plugins/VirtualProduction/DMX/DMXPixelMapping/Source/DMXPixelMappingRuntime/Private/Components/DMXPixelMappingMatrixComponent.cpp
@@ -226,6 +226,8 @@ void UDMXPixelMappingMatrixComponent::PostEditChangeProperty(FPropertyChangedEve
 	{
 		HandleMatrixChanged();
 	}
+
+	InvalidatePixelMapRenderer();
 }
 #endif // WITH_EDITOR
 
@@ -311,10 +313,10 @@ void UDMXPixelMappingMatrixComponent::SendDMX()
 			{
 				return Cast<UDMXPixelMappingMatrixCellComponent>(BaseComponent);
 			});
-		Algo::SortBy(CellComponents, [](UDMXPixelMappingMatrixCellComponent* Component)
+		Algo::SortBy(CellComponents, [ActiveModePtr](UDMXPixelMappingMatrixCellComponent* Component)
 			{
-				// Sort by cell coordinates
-				return Component->GetCellCoordinate().Y * Component->GetCellCoordinate().X + Component->GetCellCoordinate().X;
+				// Sort by cell ID
+				return Component->CellID;
 			});
 
 		// Accumulate matrix cell data
@@ -424,7 +426,6 @@ void UDMXPixelMappingMatrixComponent::SetPosition(const FVector2D& NewPosition)
 
 		}, bUpdateSizeRecursive);
 
-
 	HandlePositionChanged();
 }
 
