DMX - Fix Pixel Mapping Input Source changes do not transact, fix no initial selection without hierarchy view, fully deprecate the old pixel mapping renderer

diff --git a/Engine/Plugins/VirtualProduction/DMX/DMXPixelMapping/Source/DMXPixelMappingEditor/Private/Toolkits/DMXPixelMappingToolkit.cpp b/Engine/Plugins/VirtualProduction/DMX/DMXPixelMapping/Source/DMXPixelMappingEditor/Private/Toolkits/DMXPixelMappingToolkit.cpp
--- a/Engine/Plugins/VirtualProduction/DMX/DMXPixelMapping/Source/DMXPixelMappingEditor/Private/Toolkits/DMXPixelMappingToolkit.cpp
+++ b/Engine/Plugins/VirtualProduction/DMX/DMXPixelMapping/Source/DMXPixelMappingEditor/Private/Toolkits/DMXPixelMappingToolkit.cpp
@@ -491,6 +491,11 @@ void FDMXPixelMappingToolkit::OnComponentRemoved(UDMXPixelMapping* PixelMapping,
 
 void FDMXPixelMappingToolkit::InitializeInternal(const EToolkitMode::Type Mode, const TSharedPtr<class IToolkitHost>& InitToolkitHost, const FGuid& MessageLogGuid)
 {
+	if (!DMXPixelMapping)
+	{
+		return;
+	}
+
 	// Make sure we loaded all UObjects
 	DMXPixelMapping->CreateOrLoadObjects();
 
@@ -588,6 +593,20 @@ void FDMXPixelMappingToolkit::InitializeInternal(const EToolkitMode::Type Mode,
 	SetupCommands();
 	ExtendToolbar();
 	RegenerateMenusAndToolbars();
+	
+	// Select the first renderer component
+	if (UDMXPixelMappingRootComponent* RootComponent = DMXPixelMapping->GetRootComponent())
+	{
+		UDMXPixelMappingBaseComponent* const* FirstRendererComponetPtr = Algo::FindByPredicate(DMXPixelMapping->GetRootComponent()->GetChildren(), [](UDMXPixelMappingBaseComponent* Component)
+			{
+				return Component && Component->GetClass() == UDMXPixelMappingRendererComponent::StaticClass();
+			});
+		if (FirstRendererComponetPtr)
+		{
+			const FDMXPixelMappingComponentReference ComponentReference(StaticCastSharedRef<FDMXPixelMappingToolkit>(AsShared()), *FirstRendererComponetPtr);
+			SelectComponents(TSet<FDMXPixelMappingComponentReference>({ ComponentReference }));
+		}
+	}
 }
 
 TSharedRef<SDockTab> FDMXPixelMappingToolkit::SpawnTab_InputSourceView(const FSpawnTabArgs& Args)
diff --git a/Engine/Plugins/VirtualProduction/DMX/DMXPixelMapping/Source/DMXPixelMappingEditor/Private/Views/SDMXPixelMappingHierarchyView.cpp b/Engine/Plugins/VirtualProduction/DMX/DMXPixelMapping/Source/DMXPixelMappingEditor/Private/Views/SDMXPixelMappingHierarchyView.cpp
--- a/Engine/Plugins/VirtualProduction/DMX/DMXPixelMapping/Source/DMXPixelMappingEditor/Private/Views/SDMXPixelMappingHierarchyView.cpp
+++ b/Engine/Plugins/VirtualProduction/DMX/DMXPixelMapping/Source/DMXPixelMappingEditor/Private/Views/SDMXPixelMappingHierarchyView.cpp
@@ -356,8 +356,10 @@ void SDMXPixelMappingHierarchyView::RestoreSelectedItems()
 
 	GEditor->GetTimerManager()->SetTimerForNextTick([this]()
 		{
-
-			WidgetTreeView->ClearSelection();
+			if (WidgetTreeView.IsValid())
+			{
+				WidgetTreeView->ClearSelection();
+			}
 			for (FDMXPixelMappingHierarchyItemWidgetModelPtr& Model : RootWidgets)
 			{
 				RestoreSelectionForItemAndChildren(Model);
diff --git a/Engine/Plugins/VirtualProduction/DMX/DMXPixelMapping/Source/DMXPixelMappingEditor/Private/Views/SDMXPixelMappingInputSourceView.cpp b/Engine/Plugins/VirtualProduction/DMX/DMXPixelMapping/Source/DMXPixelMappingEditor/Private/Views/SDMXPixelMappingInputSourceView.cpp
--- a/Engine/Plugins/VirtualProduction/DMX/DMXPixelMapping/Source/DMXPixelMappingEditor/Private/Views/SDMXPixelMappingInputSourceView.cpp
+++ b/Engine/Plugins/VirtualProduction/DMX/DMXPixelMapping/Source/DMXPixelMappingEditor/Private/Views/SDMXPixelMappingInputSourceView.cpp
@@ -13,6 +13,7 @@
 #include "IDetailsView.h"
 #include "PropertyEditorDelegates.h"
 #include "PropertyEditorModule.h"
+#include "ScopedTransaction.h"
 #include "Toolkits/DMXPixelMappingToolkit.h"
 #include "Views/SDMXPixelMappingDesignerView.h"
 #include "Widgets/SBoxPanel.h"
@@ -238,9 +239,14 @@ void SDMXPixelMappingInputSourceView::OnComponentAddedOrRemoved(UDMXPixelMapping
 
 FReply SDMXPixelMappingInputSourceView::OnAddButtonClicked()
 {
-	if (const TSharedPtr<FDMXPixelMappingToolkit> Toolkit = WeakToolkit.Pin())
+	const TSharedPtr<FDMXPixelMappingToolkit> Toolkit = WeakToolkit.Pin();
+	UDMXPixelMapping* PixelMapping = Toolkit.IsValid() ? Toolkit->GetDMXPixelMapping() : nullptr;
+	if (Toolkit.IsValid() && PixelMapping && PixelMapping->GetRootComponent())
 	{
+		const FScopedTransaction AddRendererTransaction(LOCTEXT("AddRendererTransaction", "Add Pixel Mapping Input Source"));
+		PixelMapping->GetRootComponent()->PreEditChange(nullptr);
 		Toolkit->AddRenderer();
+		PixelMapping->GetRootComponent()->PostEditChange();
 		
 		UDMXPixelMappingRendererComponent* RendererComponentToSelect = Toolkit->GetActiveRendererComponent();
 		if (RendererComponentToSelect)
diff --git a/Engine/Plugins/VirtualProduction/DMX/DMXPixelMapping/Source/DMXPixelMappingRenderer/Private/DMXPixelMappingRendererModule.cpp b/Engine/Plugins/VirtualProduction/DMX/DMXPixelMapping/Source/DMXPixelMappingRenderer/Private/DMXPixelMappingRendererModule.cpp
--- a/Engine/Plugins/VirtualProduction/DMX/DMXPixelMapping/Source/DMXPixelMappingRenderer/Private/DMXPixelMappingRendererModule.cpp
+++ b/Engine/Plugins/VirtualProduction/DMX/DMXPixelMapping/Source/DMXPixelMappingRenderer/Private/DMXPixelMappingRendererModule.cpp
@@ -19,9 +19,11 @@ void FDMXPixelMappingRendererModule::StartupModule()
 	}
 }
 
+PRAGMA_DISABLE_DEPRECATION_WARNINGS
 TSharedPtr<IDMXPixelMappingRenderer> FDMXPixelMappingRendererModule::CreateRenderer() const
 {
 	return MakeShared<FDMXPixelMappingRenderer>();
 }
+PRAGMA_ENABLE_DEPRECATION_WARNINGS
 
 IMPLEMENT_MODULE(FDMXPixelMappingRendererModule, DMXPixelMappingRenderer);
diff --git a/Engine/Plugins/VirtualProduction/DMX/DMXPixelMapping/Source/DMXPixelMappingRenderer/Private/DMXPixelMappingRendererModule.h b/Engine/Plugins/VirtualProduction/DMX/DMXPixelMapping/Source/DMXPixelMappingRenderer/Private/DMXPixelMappingRendererModule.h
--- a/Engine/Plugins/VirtualProduction/DMX/DMXPixelMapping/Source/DMXPixelMappingRenderer/Private/DMXPixelMappingRendererModule.h
+++ b/Engine/Plugins/VirtualProduction/DMX/DMXPixelMapping/Source/DMXPixelMappingRenderer/Private/DMXPixelMappingRendererModule.h
@@ -4,7 +4,9 @@
 
 #include "IDMXPixelMappingRendererModule.h"
 
+PRAGMA_DISABLE_DEPRECATION_WARNINGS
 class IDMXPixelMappingRenderer;
+PRAGMA_ENABLE_DEPRECATION_WARNINGS
 
 class FDMXPixelMappingRendererModule 
 	: public IDMXPixelMappingRendererModule
@@ -16,6 +18,9 @@ public:
 	//~ End IModuleInterface implementation
 
 	//~ Begin IDMXPixelMappingRendererModule implementation
+	UE_DEPRECATED(5.3, "IDMXPixelMappingRenderer was replaced with UDMXPixelMappigPixelMapRenderer")
+	PRAGMA_DISABLE_DEPRECATION_WARNINGS
 	virtual TSharedPtr<IDMXPixelMappingRenderer> CreateRenderer() const override;
+	PRAGMA_ENABLE_DEPRECATION_WARNINGS
 	//~ End IDMXPixelMappingRendererModule implementation
 };
diff --git a/Engine/Plugins/VirtualProduction/DMX/DMXPixelMapping/Source/DMXPixelMappingRenderer/Public/IDMXPixelMappingRendererModule.h b/Engine/Plugins/VirtualProduction/DMX/DMXPixelMapping/Source/DMXPixelMappingRenderer/Public/IDMXPixelMappingRendererModule.h
--- a/Engine/Plugins/VirtualProduction/DMX/DMXPixelMapping/Source/DMXPixelMappingRenderer/Public/IDMXPixelMappingRendererModule.h
+++ b/Engine/Plugins/VirtualProduction/DMX/DMXPixelMapping/Source/DMXPixelMappingRenderer/Public/IDMXPixelMappingRendererModule.h
@@ -5,7 +5,10 @@
 #include "Modules/ModuleManager.h"
 #include "Modules/ModuleInterface.h"
 
+PRAGMA_DISABLE_DEPRECATION_WARNINGS
 class IDMXPixelMappingRenderer;
+PRAGMA_ENABLE_DEPRECATION_WARNINGS
+
 
 /**
  * Public module interface for the DMX Renderer module
@@ -49,5 +52,7 @@ public:
 	 *
 	 * @return IDMXPixelMappingRenderer instance.
 	 */
+	PRAGMA_DISABLE_DEPRECATION_WARNINGS
 	virtual TSharedPtr<IDMXPixelMappingRenderer> CreateRenderer() const = 0;
+	PRAGMA_ENABLE_DEPRECATION_WARNINGS
 };
