DMX - Now only relelvant matrix fixtures in pixel mapping update if their fixture patch changes, fix drag drop fixtures in pixel mapping does not transact for all pixel mapping components

diff --git a/Engine/Plugins/VirtualProduction/DMX/DMXPixelMapping/Source/DMXPixelMappingEditor/Private/DragDrop/DMXPixelMappingDragDropOp.cpp b/Engine/Plugins/VirtualProduction/DMX/DMXPixelMapping/Source/DMXPixelMappingEditor/Private/DragDrop/DMXPixelMappingDragDropOp.cpp
--- a/Engine/Plugins/VirtualProduction/DMX/DMXPixelMapping/Source/DMXPixelMappingEditor/Private/DragDrop/DMXPixelMappingDragDropOp.cpp
+++ b/Engine/Plugins/VirtualProduction/DMX/DMXPixelMapping/Source/DMXPixelMappingEditor/Private/DragDrop/DMXPixelMappingDragDropOp.cpp
@@ -87,6 +87,11 @@ void FDMXPixelMappingDragDropOp::LayoutOutputComponents(const FVector2D& GraphSp
 				{
 					OutputComponent->PreEditChange(nullptr);
 
+					constexpr bool bModifyChildrenRecursive = true;
+					Component->ForEachChild([](UDMXPixelMappingBaseComponent* Component)
+						{
+							Component->Modify();
+						}, bModifyChildrenRecursive);
 
 					if (ensureMsgf(OutputComponent->GetClass() != UDMXPixelMappingMatrixComponent::StaticClass(),
 						TEXT("Matrix Cells are not supported. Use the GroupChildDragDropHelper from this class instead")))
diff --git a/Engine/Plugins/VirtualProduction/DMX/DMXPixelMapping/Source/DMXPixelMappingEditor/Private/DragDrop/DMXPixelMappingGroupChildDragDropHelper.cpp b/Engine/Plugins/VirtualProduction/DMX/DMXPixelMapping/Source/DMXPixelMappingEditor/Private/DragDrop/DMXPixelMappingGroupChildDragDropHelper.cpp
--- a/Engine/Plugins/VirtualProduction/DMX/DMXPixelMapping/Source/DMXPixelMappingEditor/Private/DragDrop/DMXPixelMappingGroupChildDragDropHelper.cpp
+++ b/Engine/Plugins/VirtualProduction/DMX/DMXPixelMapping/Source/DMXPixelMappingEditor/Private/DragDrop/DMXPixelMappingGroupChildDragDropHelper.cpp
@@ -93,6 +93,12 @@ void FDMXPixelMappingGroupChildDragDropHelper::LayoutAligned(const FVector2D& Gr
 				{
 					ChildComponent->PreEditChange(nullptr);
 
+					constexpr bool bModifyChildrenRecursive = true;
+					ChildComponent->ForEachChild([](UDMXPixelMappingBaseComponent* Component)
+						{
+							Component->Modify();
+						}, bModifyChildrenRecursive);
+
 					if (GroupComponent->IsOverPosition(NextPosition) && 
 						GroupComponent->IsOverPosition(NextPosition + ChildComponent->GetSize()))
 					{
@@ -147,6 +153,12 @@ void FDMXPixelMappingGroupChildDragDropHelper::LayoutUnaligned(const FVector2D&
 					{
 						ChildComponent->PreEditChange(nullptr);
 
+						constexpr bool bModifyChildrenRecursive = true;
+						ChildComponent->ForEachChild([](UDMXPixelMappingBaseComponent* Component)
+							{
+								Component->Modify();
+							}, bModifyChildrenRecursive);
+
 						FVector2D AnchorOffset = Anchor - ChildComponent->GetPosition();
 
 						FVector2D NewPosition = GraphSpacePosition - AnchorOffset - PinnedDragDropOp->GraphSpaceDragOffset;
diff --git a/Engine/Plugins/VirtualProduction/DMX/DMXPixelMapping/Source/DMXPixelMappingRuntime/Private/Components/DMXPixelMappingMatrixComponent.cpp b/Engine/Plugins/VirtualProduction/DMX/DMXPixelMapping/Source/DMXPixelMappingRuntime/Private/Components/DMXPixelMappingMatrixComponent.cpp
--- a/Engine/Plugins/VirtualProduction/DMX/DMXPixelMapping/Source/DMXPixelMappingRuntime/Private/Components/DMXPixelMappingMatrixComponent.cpp
+++ b/Engine/Plugins/VirtualProduction/DMX/DMXPixelMapping/Source/DMXPixelMappingRuntime/Private/Components/DMXPixelMappingMatrixComponent.cpp
@@ -563,7 +563,7 @@ void UDMXPixelMappingMatrixComponent::OnFixtureTypeChanged(const UDMXEntityFixtu
 
 void UDMXPixelMappingMatrixComponent::OnFixturePatchChanged(const UDMXEntityFixturePatch* FixturePatch)
 {
-	if (FixturePatch)
+	if (FixturePatch && FixturePatch == FixturePatchRef.GetEntity())
 	{
 		HandleMatrixChanged();
 
