DMX - Fix Pixel Mapping Matrix optimized send logic ignores word size of attributes

diff --git a/Engine/Plugins/VirtualProduction/DMX/DMXPixelMapping/Source/DMXPixelMappingRuntime/Private/Components/DMXPixelMappingMatrixComponent.cpp b/Engine/Plugins/VirtualProduction/DMX/DMXPixelMapping/Source/DMXPixelMappingRuntime/Private/Components/DMXPixelMappingMatrixComponent.cpp
--- a/Engine/Plugins/VirtualProduction/DMX/DMXPixelMapping/Source/DMXPixelMappingRuntime/Private/Components/DMXPixelMappingMatrixComponent.cpp
+++ b/Engine/Plugins/VirtualProduction/DMX/DMXPixelMapping/Source/DMXPixelMappingRuntime/Private/Components/DMXPixelMappingMatrixComponent.cpp
@@ -348,18 +348,18 @@ void UDMXPixelMappingMatrixComponent::SendDMX()
 
 		int32 ChannelOffset = 0;
 		TMap<int32, uint8> ChannelToValueMap;
-		for (int32 CellIndex = 0; CellIndex < CellComponents.Num(); CellIndex++)
+		for (const FDMXNormalizedAttributeValueMap& AttributeToValueMap : AttributeToValueMapArray)
 		{
 			for (const FDMXFixtureCellAttribute& CellAttribute : ActiveModePtr->FixtureMatrixConfig.CellAttributes)
 			{
-				float* ValuePtr = AttributeToValueMapArray[CellIndex].Map.Find(CellAttribute.Attribute);
+				const float* ValuePtr = AttributeToValueMap.Map.Find(CellAttribute.Attribute);
 				if (ValuePtr)
 				{
 					const TArray<uint8> Values = FDMXConversions::NormalizedDMXValueToByteArray(*ValuePtr, CellAttribute.DataType, CellAttribute.bUseLSBMode);
-					for (uint8 Value : Values)
+					for (int32 ValueIndex = 0; ValueIndex < Values.Num(); ValueIndex++)
 					{
-						const int32 Channel = MatrixStartingChannel + ChannelOffset;
-						ChannelToValueMap.Add(Channel, Value);
+						const int32 Channel = MatrixStartingChannel + ChannelOffset + ValueIndex;
+						ChannelToValueMap.Add(Channel, Values[ValueIndex]);
 					}
 				}
 
