DMX - Fix changing color space luminance attribute to none keeps sending the to the previously set attribute

diff --git a/Engine/Plugins/VirtualProduction/DMX/DMXPixelMapping/Source/DMXPixelMappingRuntime/Private/ColorSpace/DMXPixelMappingColorSpace_RGBCMY.cpp b/Engine/Plugins/VirtualProduction/DMX/DMXPixelMapping/Source/DMXPixelMappingRuntime/Private/ColorSpace/DMXPixelMappingColorSpace_RGBCMY.cpp
--- a/Engine/Plugins/VirtualProduction/DMX/DMXPixelMapping/Source/DMXPixelMappingRuntime/Private/ColorSpace/DMXPixelMappingColorSpace_RGBCMY.cpp
+++ b/Engine/Plugins/VirtualProduction/DMX/DMXPixelMapping/Source/DMXPixelMappingRuntime/Private/ColorSpace/DMXPixelMappingColorSpace_RGBCMY.cpp
@@ -70,7 +70,9 @@ void UDMXPixelMappingColorSpace_RGBCMY::PostEditChangeChainProperty(FPropertyCha
 
 	const FName PropertyName = PropertyChangedEvent.GetPropertyName();
 
-	if (PropertyName == GET_MEMBER_NAME_CHECKED(FDMXAttributeName, Name))
+	if (PropertyName == GET_MEMBER_NAME_CHECKED(FDMXAttributeName, Name) ||
+		PropertyName == GET_MEMBER_NAME_CHECKED(UDMXPixelMappingColorSpace_RGBCMY, LuminanceType) ||
+		PropertyName == GET_MEMBER_NAME_CHECKED(UDMXPixelMappingColorSpace_RGBCMY, LuminanceAttribute))
 	{
 		ClearCachedAttributeValues();
 	}
