DMX - PixelMapping now displays its texture correctly when DPI scaling on Windows is enabled

diff --git a/Engine/Plugins/VirtualProduction/DMX/DMXPixelMapping/Source/DMXPixelMappingEditor/Private/Viewports/DMXPixelMappingSourceTextureViewportClient.cpp b/Engine/Plugins/VirtualProduction/DMX/DMXPixelMapping/Source/DMXPixelMappingEditor/Private/Viewports/DMXPixelMappingSourceTextureViewportClient.cpp
--- a/Engine/Plugins/VirtualProduction/DMX/DMXPixelMapping/Source/DMXPixelMappingEditor/Private/Viewports/DMXPixelMappingSourceTextureViewportClient.cpp
+++ b/Engine/Plugins/VirtualProduction/DMX/DMXPixelMapping/Source/DMXPixelMappingEditor/Private/Viewports/DMXPixelMappingSourceTextureViewportClient.cpp
@@ -6,6 +6,7 @@
 #include "CanvasTypes.h"
 #include "Components/DMXPixelMappingRendererComponent.h"
 #include "Engine/Texture.h"
+#include "Framework/Application/SlateApplication.h"
 #include "Slate/SceneViewport.h"
 #include "Texture2DPreview.h"
 #include "TextureResource.h"
@@ -123,3 +124,18 @@ void FDMXPixelMappingSourceTextureViewportClient::Draw(FViewport* InViewport, FC
 		Canvas->DrawItem(TileItem);
 	}
 }
+
+float FDMXPixelMappingSourceTextureViewportClient::UpdateViewportClientWindowDPIScale() const
+{
+	float DPIScale = 1.f;
+	if (WeakSourceTextureViewport.IsValid())
+	{
+		TSharedPtr<SWindow> WidgetWindow = FSlateApplication::Get().FindWidgetWindow(WeakSourceTextureViewport.Pin().ToSharedRef());
+		if (WidgetWindow.IsValid())
+		{
+			DPIScale = WidgetWindow->GetNativeWindow()->GetDPIScaleFactor();
+		}
+	}
+
+	return DPIScale;
+}
diff --git a/Engine/Plugins/VirtualProduction/DMX/DMXPixelMapping/Source/DMXPixelMappingEditor/Private/Viewports/DMXPixelMappingSourceTextureViewportClient.h b/Engine/Plugins/VirtualProduction/DMX/DMXPixelMapping/Source/DMXPixelMappingEditor/Private/Viewports/DMXPixelMappingSourceTextureViewportClient.h
--- a/Engine/Plugins/VirtualProduction/DMX/DMXPixelMapping/Source/DMXPixelMappingEditor/Private/Viewports/DMXPixelMappingSourceTextureViewportClient.h
+++ b/Engine/Plugins/VirtualProduction/DMX/DMXPixelMapping/Source/DMXPixelMappingEditor/Private/Viewports/DMXPixelMappingSourceTextureViewportClient.h
@@ -9,7 +9,7 @@ class SDMXPixelMappingSourceTextureViewport;
 class FDMXPixelMappingToolkit;
 
 class FDMXPixelMappingSourceTextureViewportClient
-	: public FViewportClient
+	: public FCommonViewportClient
 {
 public:
 	/** Constructor */
@@ -26,6 +26,10 @@ protected:
 	virtual void Draw(FViewport* Viewport, FCanvas* Canvas) override;
 	//~ End FViewportClient Interface
 
+	//~ Begin FViewport Interface
+	virtual float UpdateViewportClientWindowDPIScale() const override;
+	//~ End FViewport Interface
+
 private:
 	TWeakPtr<FDMXPixelMappingToolkit> WeakToolkit;
 	TWeakPtr<SDMXPixelMappingSourceTextureViewport> WeakSourceTextureViewport;
