DMX - Default 'show matrix cells' in PixelMapping to true

diff --git a/Engine/Plugins/VirtualProduction/DMX/DMXPixelMapping/Source/DMXPixelMappingEditor/Private/DMXPixelMappingLayoutSettings.h b/Engine/Plugins/VirtualProduction/DMX/DMXPixelMapping/Source/DMXPixelMappingEditor/Private/DMXPixelMappingLayoutSettings.h
--- a/Engine/Plugins/VirtualProduction/DMX/DMXPixelMapping/Source/DMXPixelMappingEditor/Private/DMXPixelMappingLayoutSettings.h
+++ b/Engine/Plugins/VirtualProduction/DMX/DMXPixelMapping/Source/DMXPixelMappingEditor/Private/DMXPixelMappingLayoutSettings.h
@@ -29,7 +29,7 @@ public:
 
 	/**  If true, shows a widget for each cell. It is recommended that this is turned off when pixel mapping large quantities of fixtures. */
 	UPROPERTY(Config, EditAnywhere, Category = "Layout Settings")
-	bool bShowMatrixCells = false;
+	bool bShowMatrixCells = true;
 
 	/** If true, shows the name of the Fixture Patch where applicable. It is recommended that this is turned off when pixel mapping large quantities of fixtures. */
 	UPROPERTY(Config, EditAnywhere, Category = "Layout Settings")
