DMX - Multi-select now correctly no longer shows a fixture patch list left, remove obsolete upgrade path for renderer component, remove obsolete selection methods from hierarchy view

diff --git a/Engine/Plugins/VirtualProduction/DMX/DMXPixelMapping/Source/DMXPixelMappingEditor/Private/ViewModels/DMXPixelMappingDMXLibraryViewModel.cpp b/Engine/Plugins/VirtualProduction/DMX/DMXPixelMapping/Source/DMXPixelMappingEditor/Private/ViewModels/DMXPixelMappingDMXLibraryViewModel.cpp
--- a/Engine/Plugins/VirtualProduction/DMX/DMXPixelMapping/Source/DMXPixelMappingEditor/Private/ViewModels/DMXPixelMappingDMXLibraryViewModel.cpp
+++ b/Engine/Plugins/VirtualProduction/DMX/DMXPixelMapping/Source/DMXPixelMappingEditor/Private/ViewModels/DMXPixelMappingDMXLibraryViewModel.cpp
@@ -102,11 +102,13 @@ void UDMXPixelMappingDMXLibraryViewModel::UpdateFixtureGroupFromSelection(TWeakP
 	{
 		DMXLibrary = nullptr;
 		WeakFixtureGroupComponent = nullptr;
+		bMoreThanOneFixtureGroupSelected = !FixtureGroupComponents.IsEmpty();
 	}
 	else
 	{
 		DMXLibrary = FixtureGroupComponents[0]->DMXLibrary;
 		WeakFixtureGroupComponent = FixtureGroupComponents[0];
+		bMoreThanOneFixtureGroupSelected = false;
 	}
 }
 
diff --git a/Engine/Plugins/VirtualProduction/DMX/DMXPixelMapping/Source/DMXPixelMappingEditor/Private/ViewModels/DMXPixelMappingDMXLibraryViewModel.h b/Engine/Plugins/VirtualProduction/DMX/DMXPixelMapping/Source/DMXPixelMappingEditor/Private/ViewModels/DMXPixelMappingDMXLibraryViewModel.h
--- a/Engine/Plugins/VirtualProduction/DMX/DMXPixelMapping/Source/DMXPixelMappingEditor/Private/ViewModels/DMXPixelMappingDMXLibraryViewModel.h
+++ b/Engine/Plugins/VirtualProduction/DMX/DMXPixelMapping/Source/DMXPixelMappingEditor/Private/ViewModels/DMXPixelMappingDMXLibraryViewModel.h
@@ -44,6 +44,9 @@ public:
 	/** Returns the Fixture Group, or nullptr if none or many are selected */
 	UDMXPixelMappingFixtureGroupComponent* GetFixtureGroupComponent() const { return WeakFixtureGroupComponent.Get(); }
 
+	/** Returns true if more than one fixture group is sected. */
+	bool IsMoreThanOneFixtureGroupSelected() const { return bMoreThanOneFixtureGroupSelected; }
+
 	/** Returns the default fixture patch list descriptor */
 	const FDMXReadOnlyFixturePatchListDescriptor& GetFixturePatchListDescriptor() const { return FixturePatchListDescriptor; }
 
@@ -97,6 +100,9 @@ private:
 	UPROPERTY(Config)
 	FDMXReadOnlyFixturePatchListDescriptor FixturePatchListDescriptor;
 
+	/** True if more than one fixture group is selected */
+	bool bMoreThanOneFixtureGroupSelected = false;
+
 	/** True while changing properties */
 	bool bChangingProperties = false;
 
diff --git a/Engine/Plugins/VirtualProduction/DMX/DMXPixelMapping/Source/DMXPixelMappingEditor/Private/Views/SDMXPixelMappingDMXLibraryView.cpp b/Engine/Plugins/VirtualProduction/DMX/DMXPixelMapping/Source/DMXPixelMappingEditor/Private/Views/SDMXPixelMappingDMXLibraryView.cpp
--- a/Engine/Plugins/VirtualProduction/DMX/DMXPixelMapping/Source/DMXPixelMappingEditor/Private/Views/SDMXPixelMappingDMXLibraryView.cpp
+++ b/Engine/Plugins/VirtualProduction/DMX/DMXPixelMapping/Source/DMXPixelMappingEditor/Private/Views/SDMXPixelMappingDMXLibraryView.cpp
@@ -245,18 +245,32 @@ void SDMXPixelMappingDMXLibraryView::ForceRefresh()
 	{
 		ViewModel->UpdateFixtureGroupFromSelection(WeakToolkit);
 
+		// Update the details view visibility
 		const UDMXPixelMappingFixtureGroupComponent* EditedFixtureGroup = ViewModel->GetFixtureGroupComponent();
 		const EVisibility DMXLibraryVisibility = EditedFixtureGroup ? EVisibility::Visible : EVisibility::Collapsed;
-
-		FixtureGroupNameTextBlock->SetVisibility(DMXLibraryVisibility);
 		ModelDetailsView->SetVisibility(DMXLibraryVisibility);
 
-		const FText FixtureGroupName = EditedFixtureGroup ? FText::FromString(EditedFixtureGroup->GetUserFriendlyName()) : LOCTEXT("NoFixtureGroupSelectedHint", "No Fixture Group selected");
-		FixtureGroupNameTextBlock->SetText(FixtureGroupName);
+		// Update the displayed fixture group name
+		const FText FixtureGroupNameText = [EditedFixtureGroup, this]()
+		{
+			if (EditedFixtureGroup)
+			{
+				return FText::FromString(EditedFixtureGroup->GetUserFriendlyName());
+			}
+			else if (ViewModel->IsMoreThanOneFixtureGroupSelected())
+			{
+				return LOCTEXT("ManyFixtureGroupsSelectedHint", "More than one Fixture Group selected");
+			}
+			else
+			{
+				return LOCTEXT("NoFixtureGroupSelectedHint", "No Fixture Group selected");
+			}
+		}();
+		FixtureGroupNameTextBlock->SetText(FixtureGroupNameText);
 
+		// Update the fixture patch list
 		UDMXLibrary* DMXLibrary = EditedFixtureGroup ? EditedFixtureGroup->DMXLibrary : nullptr;
 		const EVisibility FixturePatchListVisibility = DMXLibrary ? EVisibility::Visible : EVisibility::Collapsed;
-
 		FixturePatchList->SetDMXLibrary(DMXLibrary);
 
 		// Update the list to only display patches that are not added to pixel mapping
@@ -295,7 +309,6 @@ void SDMXPixelMappingDMXLibraryView::OnComponentSelected()
 	}
 
 	bRenderComponentContainedInSelection = false;
-	bool bSelectionContainsViewModel = false;
 
 	const TSet<FDMXPixelMappingComponentReference>& SelectedComponents = Toolkit->GetSelectedComponents();
 	for (const FDMXPixelMappingComponentReference& ComponentReference : SelectedComponents)
@@ -306,20 +319,11 @@ void SDMXPixelMappingDMXLibraryView::OnComponentSelected()
 			// Is renderer component contained in selection?
 			bRenderComponentContainedInSelection |= Component->GetClass() == UDMXPixelMappingRendererComponent::StaticClass();
 
-			if (ViewModel && ViewModel->GetFixtureGroupComponent() == Component)
-			{
-				// Is view model's fixture group contained in selection?
-				bSelectionContainsViewModel = true;
-			}
-
 			Component = Component->GetParent();
 		} while (Component);
 	}
 
-	if (!bSelectionContainsViewModel)
-	{
-		RequestRefresh();
-	}
+	RequestRefresh();
 }
 
 void SDMXPixelMappingDMXLibraryView::OnEntityAddedOrRemoved(UDMXLibrary* DMXLibrary, TArray<UDMXEntity*> Entities)
diff --git a/Engine/Plugins/VirtualProduction/DMX/DMXPixelMapping/Source/DMXPixelMappingEditor/Private/Views/SDMXPixelMappingHierarchyView.cpp b/Engine/Plugins/VirtualProduction/DMX/DMXPixelMapping/Source/DMXPixelMappingEditor/Private/Views/SDMXPixelMappingHierarchyView.cpp
--- a/Engine/Plugins/VirtualProduction/DMX/DMXPixelMapping/Source/DMXPixelMappingEditor/Private/Views/SDMXPixelMappingHierarchyView.cpp
+++ b/Engine/Plugins/VirtualProduction/DMX/DMXPixelMapping/Source/DMXPixelMappingEditor/Private/Views/SDMXPixelMappingHierarchyView.cpp
@@ -158,7 +158,6 @@ void SDMXPixelMappingHierarchyView::Construct(const FArguments& InArgs, const TS
 		];
 
 	bRebuildTreeRequested = true;
-	bSelectFirstRenderer = true;
 
 	GEditor->RegisterForUndo(this);
 
@@ -170,31 +169,9 @@ void SDMXPixelMappingHierarchyView::Construct(const FArguments& InArgs, const TS
 	UDMXPixelMappingBaseComponent::GetOnComponentRemoved().AddSP(this, &SDMXPixelMappingHierarchyView::OnComponentRemoved);
 }
 
-void SDMXPixelMappingHierarchyView::SelectFirstAvailableRenderer()
-{
-	bSelectFirstRenderer = true;
-}
-
 void SDMXPixelMappingHierarchyView::Tick(const FGeometry& AllottedGeometry, const double InCurrentTime, const float InDeltaTime)
 {
 	ConditionallyUpdateTree();
-
-	// select the first available renderer, if any is available
-	if (bSelectFirstRenderer && TreeRootWidgets.Num() > 0 && WidgetTreeView->GetNumItemsSelected() == 0)
-	{
-		FDMXPixelMappingHierarchyItemWidgetModelArr RootComponents;
-		TreeRootWidgets[0]->GatherChildren(RootComponents);
-		for (FDMXPixelMappingHierarchyItemWidgetModelPtr ItemWidgetPtr : RootComponents)
-		{
-			if (UDMXPixelMappingRendererComponent* RendererComponent = Cast<UDMXPixelMappingRendererComponent>(ItemWidgetPtr->GetReference().GetComponent()))
-			{
-				WidgetTreeView->SetSelection(ItemWidgetPtr, ESelectInfo::OnMouseClick);
-				bSelectFirstRenderer = false;
-
-				break;
-			}
-		}
-	}
 }
 
 FReply SDMXPixelMappingHierarchyView::OnKeyDown(const FGeometry& MyGeometry, const FKeyEvent& InKeyEvent)
diff --git a/Engine/Plugins/VirtualProduction/DMX/DMXPixelMapping/Source/DMXPixelMappingEditor/Private/Views/SDMXPixelMappingHierarchyView.h b/Engine/Plugins/VirtualProduction/DMX/DMXPixelMapping/Source/DMXPixelMappingEditor/Private/Views/SDMXPixelMappingHierarchyView.h
--- a/Engine/Plugins/VirtualProduction/DMX/DMXPixelMapping/Source/DMXPixelMappingEditor/Private/Views/SDMXPixelMappingHierarchyView.h
+++ b/Engine/Plugins/VirtualProduction/DMX/DMXPixelMapping/Source/DMXPixelMappingEditor/Private/Views/SDMXPixelMappingHierarchyView.h
@@ -127,8 +127,6 @@ private:
 
 	bool MoveComponentToComponent(UDMXPixelMappingBaseComponent* Source, UDMXPixelMappingBaseComponent* Destination, const bool bRename);
 
-	void SelectFirstAvailableRenderer();
-
 private:
 	TWeakPtr<FDMXPixelMappingToolkit> WeakToolkit;
 
@@ -160,6 +158,4 @@ private:
 	bool bIsUpdatingSelection;
 
 	TreeViewPtr WidgetTreeView;
-
-	bool bSelectFirstRenderer;
 };
diff --git a/Engine/Plugins/VirtualProduction/DMX/DMXPixelMapping/Source/DMXPixelMappingRuntime/Private/Components/DMXPixelMappingRendererComponent.cpp b/Engine/Plugins/VirtualProduction/DMX/DMXPixelMapping/Source/DMXPixelMappingRuntime/Private/Components/DMXPixelMappingRendererComponent.cpp
--- a/Engine/Plugins/VirtualProduction/DMX/DMXPixelMapping/Source/DMXPixelMappingRuntime/Private/Components/DMXPixelMappingRendererComponent.cpp
+++ b/Engine/Plugins/VirtualProduction/DMX/DMXPixelMapping/Source/DMXPixelMappingRuntime/Private/Components/DMXPixelMappingRendererComponent.cpp
@@ -71,31 +71,6 @@ bool UDMXPixelMappingRendererComponent::CanBeMovedTo(const UDMXPixelMappingBaseC
 	return Component && Component->GetClass() == UDMXPixelMappingRootComponent::StaticClass();
 }
 
-void UDMXPixelMappingRendererComponent::Serialize(FArchive& Ar)
-{
-	Super::Serialize(Ar);
-
-	Ar.UsingCustomVersion(FDMXPixelMappingMainStreamObjectVersion::GUID);
-	if (Ar.IsLoading())
-	{
-		if (Ar.CustomVer(FDMXPixelMappingMainStreamObjectVersion::GUID) < FDMXPixelMappingMainStreamObjectVersion::LockRendererComponentsThatUseTextureInDesigner)
-		{
-			if (RendererType == EDMXPixelMappingRendererType::Texture)
-			{
-				// Refresh the size of the texture if that is used as input
-				if (InputTexture)
-				{
-					if (const FTextureResource* TextureResource = InputTexture->GetResource())
-					{
-						const FVector2D NewSize = FVector2D(TextureResource->GetSizeX(), TextureResource->GetSizeY());
-						SetSize(NewSize);
-					}
-				}
-			}
-		}
-	}
-}
-
 void UDMXPixelMappingRendererComponent::PostInitProperties()
 {
 	Super::PostInitProperties();
diff --git a/Engine/Plugins/VirtualProduction/DMX/DMXPixelMapping/Source/DMXPixelMappingRuntime/Public/Components/DMXPixelMappingRendererComponent.h b/Engine/Plugins/VirtualProduction/DMX/DMXPixelMapping/Source/DMXPixelMappingRuntime/Public/Components/DMXPixelMappingRendererComponent.h
--- a/Engine/Plugins/VirtualProduction/DMX/DMXPixelMapping/Source/DMXPixelMappingRuntime/Public/Components/DMXPixelMappingRendererComponent.h
+++ b/Engine/Plugins/VirtualProduction/DMX/DMXPixelMapping/Source/DMXPixelMappingRuntime/Public/Components/DMXPixelMappingRendererComponent.h
@@ -37,7 +37,6 @@ public:
 
 	//~ Begin UObject implementation
 protected:
-	virtual void Serialize(FArchive& Ar) override;
 	virtual void PostInitProperties() override;
 	virtual void PostLoad() override;
 #if WITH_EDITOR
