DMX - Move common properties of matrix and fixture group item to their base class to allow multi editing, improve categories and ordering, unify Pixel Mapping Quality Setting enum to a single declaration

diff --git a/Engine/Plugins/VirtualProduction/DMX/DMXPixelMapping/Source/DMXPixelMappingEditor/Private/Customizations/DMXPixelMappingDetailCustomization_FixtureGroupItem.cpp b/Engine/Plugins/VirtualProduction/DMX/DMXPixelMapping/Source/DMXPixelMappingEditor/Private/Customizations/DMXPixelMappingDetailCustomization_FixtureGroupItem.cpp
--- a/Engine/Plugins/VirtualProduction/DMX/DMXPixelMapping/Source/DMXPixelMappingEditor/Private/Customizations/DMXPixelMappingDetailCustomization_FixtureGroupItem.cpp
+++ b/Engine/Plugins/VirtualProduction/DMX/DMXPixelMapping/Source/DMXPixelMappingEditor/Private/Customizations/DMXPixelMappingDetailCustomization_FixtureGroupItem.cpp
@@ -49,42 +49,67 @@ void FDMXPixelMappingDetailCustomization_FixtureGroupItem::CustomizeDetails(IDet
 	// Sort categories
 	InDetailLayout.SortCategories([](const TMap<FName, IDetailCategoryBuilder*>& CategoryMap)
 		{
-			int32 MinSortOrder = TNumericLimits<int32>::Max();
-			for (const TPair<FName, IDetailCategoryBuilder*>& Pair : CategoryMap)
+			int32 SortOrder = 0;
+
+			IDetailCategoryBuilder* const* FixturePatchCategoryPtr = CategoryMap.Find("Fixture Patch");
+			if (FixturePatchCategoryPtr)
 			{
-				int32 SortOrder = Pair.Value->GetSortOrder();
-				MinSortOrder = FMath::Min(SortOrder, MinSortOrder);
+				(*FixturePatchCategoryPtr)->SetSortOrder(SortOrder++);
 			}
 
 			IDetailCategoryBuilder* const* ColorSpaceCategoryPtr = CategoryMap.Find("Color Space");
 			if (ColorSpaceCategoryPtr)
 			{
-				(*ColorSpaceCategoryPtr)->SetSortOrder(MinSortOrder - 3);
+				(*ColorSpaceCategoryPtr)->SetSortOrder(SortOrder++);
 			}
 
 			// Either 'RGB' 'XY' or 'XYZ' is displayed
 			IDetailCategoryBuilder* const* RGBCategoryPtr = CategoryMap.Find("RGB");
 			if (RGBCategoryPtr)
 			{
-				(*RGBCategoryPtr)->SetSortOrder(MinSortOrder - 2);
+				(*RGBCategoryPtr)->SetSortOrder(SortOrder++);
 			}
 
 			IDetailCategoryBuilder* const* XYCategoryPtr = CategoryMap.Find("XY");
 			if (XYCategoryPtr)
 			{
-				(*XYCategoryPtr)->SetSortOrder(MinSortOrder - 2);
+				(*XYCategoryPtr)->SetSortOrder(SortOrder++);
 			}
 
 			IDetailCategoryBuilder* const* XYZCategoryPtr = CategoryMap.Find("XYZ");
 			if (XYZCategoryPtr)
 			{
-				(*XYZCategoryPtr)->SetSortOrder(MinSortOrder - 2);
+				(*XYZCategoryPtr)->SetSortOrder(SortOrder++);
 			}
 
 			IDetailCategoryBuilder* const* IntensityCategoryPtr = CategoryMap.Find("Luminance");
 			if (IntensityCategoryPtr)
 			{
-				(*IntensityCategoryPtr)->SetSortOrder(MinSortOrder - 1);
+				(*IntensityCategoryPtr)->SetSortOrder(SortOrder++);
+			}
+
+			IDetailCategoryBuilder* const* QualityCategoryPtr = CategoryMap.Find("Quality");
+			if (QualityCategoryPtr)
+			{
+				(*QualityCategoryPtr)->SetSortOrder(SortOrder++);
+			}
+
+			IDetailCategoryBuilder* const* ModulatorsCategoryPtr = CategoryMap.Find("Output Modulators");
+			if (ModulatorsCategoryPtr)
+			{
+				(*ModulatorsCategoryPtr)->SetSortOrder(SortOrder++);
+			}
+
+			IDetailCategoryBuilder* const* TransformCategoryPtr = CategoryMap.Find("Transform");
+			if (TransformCategoryPtr)
+			{
+				(*TransformCategoryPtr)->SetSortOrder(SortOrder++);
+			}
+
+			IDetailCategoryBuilder* const* EditorSettingsCategoryPtr = CategoryMap.Find("Editor Settings");
+			if (EditorSettingsCategoryPtr)
+			{
+				(*EditorSettingsCategoryPtr)->SetSortOrder(SortOrder++);
 			}
 		});
 }
diff --git a/Engine/Plugins/VirtualProduction/DMX/DMXPixelMapping/Source/DMXPixelMappingEditor/Private/Customizations/DMXPixelMappingDetailCustomization_Matrix.cpp b/Engine/Plugins/VirtualProduction/DMX/DMXPixelMapping/Source/DMXPixelMappingEditor/Private/Customizations/DMXPixelMappingDetailCustomization_Matrix.cpp
--- a/Engine/Plugins/VirtualProduction/DMX/DMXPixelMapping/Source/DMXPixelMappingEditor/Private/Customizations/DMXPixelMappingDetailCustomization_Matrix.cpp
+++ b/Engine/Plugins/VirtualProduction/DMX/DMXPixelMapping/Source/DMXPixelMappingEditor/Private/Customizations/DMXPixelMappingDetailCustomization_Matrix.cpp
@@ -44,42 +44,67 @@ void FDMXPixelMappingDetailCustomization_Matrix::CustomizeDetails(IDetailLayoutB
 	// Sort categories
 	InDetailLayout.SortCategories([](const TMap<FName, IDetailCategoryBuilder*>& CategoryMap)
 		{
-			int32 MinSortOrder = TNumericLimits<int32>::Max();
-			for (const TPair<FName, IDetailCategoryBuilder*>& Pair : CategoryMap)
+			int32 SortOrder = 0;
+
+			IDetailCategoryBuilder* const* FixturePatchCategoryPtr = CategoryMap.Find("Fixture Patch");
+			if (FixturePatchCategoryPtr)
 			{
-				int32 SortOrder = Pair.Value->GetSortOrder();
-				MinSortOrder = FMath::Min(SortOrder, MinSortOrder);
+				(*FixturePatchCategoryPtr)->SetSortOrder(SortOrder++);
 			}
 
 			IDetailCategoryBuilder* const* ColorSpaceCategoryPtr = CategoryMap.Find("Color Space");
 			if (ColorSpaceCategoryPtr)
 			{
-				(*ColorSpaceCategoryPtr)->SetSortOrder(MinSortOrder - 3);
+				(*ColorSpaceCategoryPtr)->SetSortOrder(SortOrder++);
 			}
 
 			// Either 'RGB' 'XY' or 'XYZ' is displayed
 			IDetailCategoryBuilder* const* RGBCategoryPtr = CategoryMap.Find("RGB");
 			if (RGBCategoryPtr)
 			{
-				(*RGBCategoryPtr)->SetSortOrder(MinSortOrder - 2);
+				(*RGBCategoryPtr)->SetSortOrder(SortOrder++);
 			}
 
 			IDetailCategoryBuilder* const* XYCategoryPtr = CategoryMap.Find("XY");
 			if (XYCategoryPtr)
 			{
-				(*XYCategoryPtr)->SetSortOrder(MinSortOrder - 2);
+				(*XYCategoryPtr)->SetSortOrder(SortOrder++);
 			}
 
 			IDetailCategoryBuilder* const* XYZCategoryPtr = CategoryMap.Find("XYZ");
 			if (XYZCategoryPtr)
 			{
-				(*XYZCategoryPtr)->SetSortOrder(MinSortOrder - 2);
+				(*XYZCategoryPtr)->SetSortOrder(SortOrder++);
 			}
 
 			IDetailCategoryBuilder* const* IntensityCategoryPtr = CategoryMap.Find("Luminance");
 			if (IntensityCategoryPtr)
 			{
-				(*IntensityCategoryPtr)->SetSortOrder(MinSortOrder - 1);
+				(*IntensityCategoryPtr)->SetSortOrder(SortOrder++);
+			}
+
+			IDetailCategoryBuilder* const* QualityCategoryPtr = CategoryMap.Find("Quality");
+			if (QualityCategoryPtr)
+			{
+				(*QualityCategoryPtr)->SetSortOrder(SortOrder++);
+			}
+
+			IDetailCategoryBuilder* const* ModulatorsCategoryPtr = CategoryMap.Find("Output Modulators");
+			if (ModulatorsCategoryPtr)
+			{
+				(*ModulatorsCategoryPtr)->SetSortOrder(SortOrder++);
+			}
+
+			IDetailCategoryBuilder* const* TransformCategoryPtr = CategoryMap.Find("Transform");
+			if (TransformCategoryPtr)
+			{
+				(*TransformCategoryPtr)->SetSortOrder(SortOrder++);
+			}
+
+			IDetailCategoryBuilder* const* EditorSettingsCategoryPtr = CategoryMap.Find("Editor Settings");
+			if (EditorSettingsCategoryPtr)
+			{
+				(*EditorSettingsCategoryPtr)->SetSortOrder(SortOrder++);
 			}
 		});
 }
diff --git a/Engine/Plugins/VirtualProduction/DMX/DMXPixelMapping/Source/DMXPixelMappingEditor/Private/Customizations/DMXPixelMappingDetailCustomization_Renderer.cpp b/Engine/Plugins/VirtualProduction/DMX/DMXPixelMapping/Source/DMXPixelMappingEditor/Private/Customizations/DMXPixelMappingDetailCustomization_Renderer.cpp
--- a/Engine/Plugins/VirtualProduction/DMX/DMXPixelMapping/Source/DMXPixelMappingEditor/Private/Customizations/DMXPixelMappingDetailCustomization_Renderer.cpp
+++ b/Engine/Plugins/VirtualProduction/DMX/DMXPixelMapping/Source/DMXPixelMappingEditor/Private/Customizations/DMXPixelMappingDetailCustomization_Renderer.cpp
@@ -50,9 +50,6 @@ void FDMXPixelMappingDetailCustomization_Renderer::CustomizeDetails(IDetailLayou
 	const TSharedRef<IPropertyHandle> IsVisibleInDesignerHandle = DetailLayout.GetProperty(UDMXPixelMappingOutputComponent::GetVisibleInDesignerPropertyName(), UDMXPixelMappingOutputComponent::StaticClass());
 	DetailLayout.HideProperty(IsVisibleInDesignerHandle);
 
-	const TSharedRef<IPropertyHandle> PixelBlendingHandle = DetailLayout.GetProperty(GET_MEMBER_NAME_CHECKED(UDMXPixelMappingOutputComponent, CellBlendingQuality), UDMXPixelMappingOutputComponent::StaticClass());
-	DetailLayout.HideProperty(PixelBlendingHandle);
-
 	// Add properties
 	IDetailCategoryBuilder& RenderSettingsCategory = DetailLayout.EditCategory("Render Settings", FText::GetEmpty(), ECategoryPriority::Important);
 	RenderSettingsCategory.AddProperty(RendererTypeHandle);
diff --git a/Engine/Plugins/VirtualProduction/DMX/DMXPixelMapping/Source/DMXPixelMappingEditor/Private/Customizations/DMXPixelMappingDetailCustomization_Screen.cpp b/Engine/Plugins/VirtualProduction/DMX/DMXPixelMapping/Source/DMXPixelMappingEditor/Private/Customizations/DMXPixelMappingDetailCustomization_Screen.cpp
--- a/Engine/Plugins/VirtualProduction/DMX/DMXPixelMapping/Source/DMXPixelMappingEditor/Private/Customizations/DMXPixelMappingDetailCustomization_Screen.cpp
+++ b/Engine/Plugins/VirtualProduction/DMX/DMXPixelMapping/Source/DMXPixelMappingEditor/Private/Customizations/DMXPixelMappingDetailCustomization_Screen.cpp
@@ -2,13 +2,18 @@
 
 #include "Customizations/DMXPixelMappingDetailCustomization_Screen.h"
 
+#include "Components/DMXPixelMappingOutputDMXComponent.h"
 #include "DetailLayoutBuilder.h"
 
 #define LOCTEXT_NAMESPACE "DMXPixelMappingDetailCustomization_Screen"
 
 void FDMXPixelMappingDetailCustomization_Screen::CustomizeDetails(IDetailLayoutBuilder& DetailLayout)
 {
-	DetailLayout.EditCategory("Patch Settings", FText::GetEmpty(), ECategoryPriority::Important);
+	// Hide pointless properties of the common base class with fixture group item and matrix component while retaining support for this, 5.3 Deprecated, component
+	DetailLayout.HideProperty(GET_MEMBER_NAME_CHECKED(UDMXPixelMappingOutputDMXComponent, FixturePatchRef));
+	DetailLayout.HideProperty(GET_MEMBER_NAME_CHECKED(UDMXPixelMappingOutputDMXComponent, ColorSpaceClass));
+	DetailLayout.HideProperty(GET_MEMBER_NAME_CHECKED(UDMXPixelMappingOutputDMXComponent, ColorSpace));
+	DetailLayout.HideProperty(GET_MEMBER_NAME_CHECKED(UDMXPixelMappingOutputDMXComponent, ModulatorClasses));
 }
 
 #undef LOCTEXT_NAMESPACE
diff --git a/Engine/Plugins/VirtualProduction/DMX/DMXPixelMapping/Source/DMXPixelMappingEditor/Private/Toolkits/DMXPixelMappingToolkit.cpp b/Engine/Plugins/VirtualProduction/DMX/DMXPixelMapping/Source/DMXPixelMappingEditor/Private/Toolkits/DMXPixelMappingToolkit.cpp
--- a/Engine/Plugins/VirtualProduction/DMX/DMXPixelMapping/Source/DMXPixelMappingEditor/Private/Toolkits/DMXPixelMappingToolkit.cpp
+++ b/Engine/Plugins/VirtualProduction/DMX/DMXPixelMapping/Source/DMXPixelMappingEditor/Private/Toolkits/DMXPixelMappingToolkit.cpp
@@ -503,7 +503,7 @@ void FDMXPixelMappingToolkit::InitializeInternal(const EToolkitMode::Type Mode,
 
 	CreateInternalViews();
 
-	TSharedRef<FTabManager::FLayout> StandaloneDefaultLayout = FTabManager::NewLayout("Standalone_PixelMapping_Layout_v9")
+	TSharedRef<FTabManager::FLayout> StandaloneDefaultLayout = FTabManager::NewLayout("Standalone_PixelMapping_Layout_2.0")
 		->AddArea
 		(
 			FTabManager::NewPrimaryArea()->SetOrientation(Orient_Vertical)
@@ -511,63 +511,60 @@ void FDMXPixelMappingToolkit::InitializeInternal(const EToolkitMode::Type Mode,
 			(
 				FTabManager::NewSplitter()
 				->SetOrientation(Orient_Horizontal)
+				->SetSizeCoefficient(.25f)
 				->Split
 				(
 					FTabManager::NewSplitter()
 					->SetOrientation(Orient_Vertical)
-					->SetSizeCoefficient(0.25f)
 					->Split
 					(
 						FTabManager::NewStack()
-						->AddTab(DMXLibraryViewTabID, ETabState::OpenedTab)
-						->SetSizeCoefficient(0.382f)
+						->AddTab(HierarchyViewTabID, ETabState::OpenedTab)
+						->SetSizeCoefficient(.618f)
 					)
 					->Split
 					(
 						FTabManager::NewStack()
-						->AddTab(HierarchyViewTabID, ETabState::OpenedTab)
-						->SetSizeCoefficient(0.618)
+						->AddTab(DMXLibraryViewTabID, ETabState::OpenedTab)
+						->SetSizeCoefficient(.382f)
 					)
 				)
-				->SetSizeCoefficient(0.25f)
+				->SetSizeCoefficient(.5f)
 				->Split
 				(
 					FTabManager::NewSplitter()
 					->SetOrientation(Orient_Vertical)
-					->SetSizeCoefficient(0.5f)
 					->Split
 					(
 						FTabManager::NewStack()
 						->AddTab(DesignerViewTabID, ETabState::OpenedTab)
-						->SetSizeCoefficient(0.6f)
+						->SetSizeCoefficient(.618f)
 					)
 					->Split
 					(
 						FTabManager::NewStack()
 						->AddTab(PreviewViewTabID, ETabState::OpenedTab)
-						->SetSizeCoefficient(0.4f)
+						->SetSizeCoefficient(.382f)
 					)
 				)
-				->SetSizeCoefficient(0.5f)
+				->SetSizeCoefficient(.25f)
 				->Split
 				(
 					FTabManager::NewSplitter()
 					->SetOrientation(Orient_Vertical)
-					->SetSizeCoefficient(0.5f)
 					->Split
 					(
 						FTabManager::NewStack()
 						->AddTab(DetailsViewTabID, ETabState::OpenedTab)
-						->SetSizeCoefficient(0.618f)
+						->SetSizeCoefficient(.618f)
 					)
 					->Split
 					(
 						FTabManager::NewStack()
 						->AddTab(LayoutViewTabID, ETabState::OpenedTab)
-						->SetSizeCoefficient(0.382f)
+						->SetSizeCoefficient(.382f)
 					)
 				)
-				->SetSizeCoefficient(0.25f)
 			)
 		);
 
@@ -593,6 +590,9 @@ void FDMXPixelMappingToolkit::InitializeInternal(const EToolkitMode::Type Mode,
 			SelectComponents(TSet<FDMXPixelMappingComponentReference>({ ComponentReference }));
 		}
 	}
+
+	FDMXPixelMappingComponentTemplate Template(UDMXPixelMappingScreenComponent::StaticClass());
+	GetActiveRendererComponent()->AddChild(Template.CreateComponent< UDMXPixelMappingScreenComponent>(DMXPixelMapping->GetRootComponent()));
 }
 
 TSharedRef<SDockTab> FDMXPixelMappingToolkit::SpawnTab_DMXLibraryView(const FSpawnTabArgs& Args)
diff --git a/Engine/Plugins/VirtualProduction/DMX/DMXPixelMapping/Source/DMXPixelMappingEditor/Private/Viewports/DMXPixelMappingPreviewViewportClient.cpp b/Engine/Plugins/VirtualProduction/DMX/DMXPixelMapping/Source/DMXPixelMappingEditor/Private/Viewports/DMXPixelMappingPreviewViewportClient.cpp
--- a/Engine/Plugins/VirtualProduction/DMX/DMXPixelMapping/Source/DMXPixelMappingEditor/Private/Viewports/DMXPixelMappingPreviewViewportClient.cpp
+++ b/Engine/Plugins/VirtualProduction/DMX/DMXPixelMapping/Source/DMXPixelMappingEditor/Private/Viewports/DMXPixelMappingPreviewViewportClient.cpp
@@ -5,8 +5,8 @@
 #include "CanvasItem.h"
 #include "CanvasTypes.h"
 #include "Components/DMXPixelMappingRendererComponent.h"
+#include "DMXPixelMappingRenderElement.h"
 #include "Engine/TextureRenderTarget2D.h"
-#include "PixelMapRenderElement.h"
 #include "Slate/SceneViewport.h"
 #include "Toolkits/DMXPixelMappingToolkit.h"
 #include "UnrealClient.h"
diff --git a/Engine/Plugins/VirtualProduction/DMX/DMXPixelMapping/Source/DMXPixelMappingRenderer/Private/DMXPixelMapRenderElement.cpp b/Engine/Plugins/VirtualProduction/DMX/DMXPixelMapping/Source/DMXPixelMappingRenderer/Private/DMXPixelMapRenderElement.cpp
new file mode 100644
--- /dev/null
+++ b/Engine/Plugins/VirtualProduction/DMX/DMXPixelMapping/Source/DMXPixelMappingRenderer/Private/DMXPixelMapRenderElement.cpp
@@ -0,0 +1,39 @@
+// Copyright Epic Games, Inc. All Rights Reserved.
+
+#include "DMXPixelMappingRenderElement.h"
+
+
+namespace UE::DMXPixelMapping::Rendering
+{
+	void FPixelMapRenderElement::SetParameters(FPixelMapRenderElementParameters InParameters)
+	{
+#if !defined(_MSC_VER) 			// Atomic for MSVC only
+		const FScopeLock LockParameters(&AccessParametersMutex);
+#endif
+		Parameters = InParameters;
+	}
+
+	FPixelMapRenderElementParameters FPixelMapRenderElement::GetParameters() const
+	{
+#if !defined(_MSC_VER)  		// Atomic for MSVC only
+		const FScopeLock LockParameters(&AccessParametersMutex);
+#endif
+		return Parameters;
+	}
+
+	void FPixelMapRenderElement::SetColor(const FLinearColor& InColor)
+	{
+#if !defined(_MSC_VER)  		// Atomic for MSVC only
+		const FScopeLock LockColorMutex(&AccessColorMutex);
+#endif
+		Color = InColor;
+	}
+
+	FLinearColor FPixelMapRenderElement::GetColor() const
+	{
+#if !defined(_MSC_VER)  		// Atomic for MSVC only
+		const FScopeLock LockColorMutex(&AccessColorMutex);
+#endif
+		return Color;
+	}
+}
diff --git a/Engine/Plugins/VirtualProduction/DMX/DMXPixelMapping/Source/DMXPixelMappingRenderer/Private/DMXPixelMappingPixelMapRenderer.cpp b/Engine/Plugins/VirtualProduction/DMX/DMXPixelMapping/Source/DMXPixelMappingRenderer/Private/DMXPixelMappingPixelMapRenderer.cpp
--- a/Engine/Plugins/VirtualProduction/DMX/DMXPixelMapping/Source/DMXPixelMappingRenderer/Private/DMXPixelMappingPixelMapRenderer.cpp
+++ b/Engine/Plugins/VirtualProduction/DMX/DMXPixelMapping/Source/DMXPixelMappingRenderer/Private/DMXPixelMappingPixelMapRenderer.cpp
@@ -2,8 +2,8 @@
 
 #include "DMXPixelMappingPixelMapRenderer.h"
 
+#include "DMXPixelMappingRenderElement.h"
 #include "Engine/TextureRenderTarget2D.h"
-#include "PixelMapRenderElement.h"
 #include "RenderPixelMapProxy.h"
 #include "RenderingThread.h"
 
diff --git a/Engine/Plugins/VirtualProduction/DMX/DMXPixelMapping/Source/DMXPixelMappingRenderer/Private/PixelMapRenderElement.cpp b/Engine/Plugins/VirtualProduction/DMX/DMXPixelMapping/Source/DMXPixelMappingRenderer/Private/PixelMapRenderElement.cpp
deleted file mode 100644
diff --git a/Engine/Plugins/VirtualProduction/DMX/DMXPixelMapping/Source/DMXPixelMappingRenderer/Private/RenderPixelMapProxy.cpp b/Engine/Plugins/VirtualProduction/DMX/DMXPixelMapping/Source/DMXPixelMappingRenderer/Private/RenderPixelMapProxy.cpp
--- a/Engine/Plugins/VirtualProduction/DMX/DMXPixelMapping/Source/DMXPixelMappingRenderer/Private/RenderPixelMapProxy.cpp
+++ b/Engine/Plugins/VirtualProduction/DMX/DMXPixelMapping/Source/DMXPixelMappingRenderer/Private/RenderPixelMapProxy.cpp
@@ -2,8 +2,8 @@
 
 #include "RenderPixelMapProxy.h"
 
+#include "DMXPixelMappingRenderElement.h"
 #include "Engine/TextureRenderTarget2D.h"
-#include "PixelMapRenderElement.h"
 #include "PixelShaderUtils.h"
 #include "ProfilingDebugging/RealtimeGPUProfiler.h"
 #include "RenderingThread.h"
@@ -19,7 +19,7 @@ namespace UE::DMXPixelMapping::Rendering::Private
 
 DECLARE_GPU_STAT_NAMED(DMXPixelMappingShadersStat, UE::DMXPixelMapping::Rendering::Private::RenderPassHint);
 
-class FDMXPixelBlendingQualityDimension : SHADER_PERMUTATION_ENUM_CLASS("PIXELBLENDING_QUALITY", UE::DMXPixelMapping::Rendering::ECellBlendingQuality);
+class FDMXPixelBlendingQualityDimension : SHADER_PERMUTATION_ENUM_CLASS("PIXELBLENDING_QUALITY", EDMXPixelBlendingQuality);
 class FDMXVertexUVDimension : SHADER_PERMUTATION_BOOL("VERTEX_UV_STATIC_CALCULATION");
 
 /**
diff --git a/Engine/Plugins/VirtualProduction/DMX/DMXPixelMapping/Source/DMXPixelMappingRenderer/Public/DMXPixelMappingPixelMapRendererParams.h b/Engine/Plugins/VirtualProduction/DMX/DMXPixelMapping/Source/DMXPixelMappingRenderer/Public/DMXPixelMappingPixelMapRendererParams.h
deleted file mode 100644
diff --git a/Engine/Plugins/VirtualProduction/DMX/DMXPixelMapping/Source/DMXPixelMappingRenderer/Public/DMXPixelMappingRenderElement.h b/Engine/Plugins/VirtualProduction/DMX/DMXPixelMapping/Source/DMXPixelMappingRenderer/Public/DMXPixelMappingRenderElement.h
new file mode 100644
--- /dev/null
+++ b/Engine/Plugins/VirtualProduction/DMX/DMXPixelMapping/Source/DMXPixelMappingRenderer/Public/DMXPixelMappingRenderElement.h
@@ -0,0 +1,99 @@
+// Copyright Epic Games, Inc. All Rights Reserved.
+
+#pragma once
+
+#if defined(_MSC_VER) 
+#include <atomic>
+#else
+#include "HAL/CriticalSection.h"
+#endif
+#include "Math/Color.h"
+#include "Math/Vector2D.h"
+
+#include "DMXPixelMappingRenderElement.generated.h"
+
+
+/** Enum that defines the quality of how pixels are rendered */
+UENUM()
+enum class EDMXPixelBlendingQuality : uint8
+{
+	/** 1 sample */
+	Low,
+
+	/** 5 samples ( 2 x 2 with center) */
+	Medium,
+
+	/** 9 samples ( 3 x 3 ) */
+	High,
+
+	/** Max, for shaders only */
+	MAX UMETA(Hidden)
+};
+
+namespace UE::DMXPixelMapping::Rendering
+{
+	struct DMXPIXELMAPPINGRENDERER_API FPixelMapRenderElementParameters
+	{
+		/** Position in texels of the top left corner of the quad's UV's */
+		FVector2D UV;
+
+		/** Size in texels of the quad's total UV space */
+		FVector2D UVSize;
+
+		/** Size in texels of UV.May match UVSize */
+		FVector2D UVCellSize;
+
+		/** The quality of color samples in the pixel shader(number of samples) */
+		EDMXPixelBlendingQuality CellBlendingQuality;
+
+		/** Calculates the UV point to sample purely on the UV position/size. Works best for renderers which represent a single pixel */
+		bool bStaticCalculateUV = true;
+	};
+
+	/** An element rendered by the pixel map renderer */
+	class DMXPIXELMAPPINGRENDERER_API FPixelMapRenderElement
+		: public TSharedFromThis<FPixelMapRenderElement>
+	{
+	public:
+		FPixelMapRenderElement(FPixelMapRenderElementParameters InParameters)
+			: Parameters(InParameters)
+		{}
+
+		virtual ~FPixelMapRenderElement() {};
+
+		/** Sets the parameters, thread-safe */
+		void SetParameters(FPixelMapRenderElementParameters InParameters);
+
+		/** Gets the parameters, thread-safe */
+		FPixelMapRenderElementParameters GetParameters() const;
+
+		/** Sets the pixel color, thread safe. This is called from the renderer to set the rendered color. */
+		void SetColor(const FLinearColor& Color);
+
+		/** Gets the pixel color, thread safe */
+		FLinearColor GetColor() const;
+
+	private:
+		// 5.3: Significantly better performance with atomics on MSVC. 
+		// However not all compilers currently support non-lockfree atomics.
+#if defined(_MSC_VER) 
+		/** The current parameters */
+		std::atomic<FPixelMapRenderElementParameters> Parameters;
+
+		/** The current pixel color */
+		std::atomic<FLinearColor> Color;
+#else
+		/** The current parameters */
+		FPixelMapRenderElementParameters Parameters;
+
+		/** The current pixel color */
+		FLinearColor Color;
+
+		/** Mutex acccess to parameters */
+		mutable FCriticalSection AccessParametersMutex;
+
+		/** Mutex acccess to parameters */
+		mutable FCriticalSection AccessColorMutex;
+#endif
+	};
+}
diff --git a/Engine/Plugins/VirtualProduction/DMX/DMXPixelMapping/Source/DMXPixelMappingRenderer/Public/PixelMapRenderElement.h b/Engine/Plugins/VirtualProduction/DMX/DMXPixelMapping/Source/DMXPixelMappingRenderer/Public/PixelMapRenderElement.h
deleted file mode 100644
diff --git a/Engine/Plugins/VirtualProduction/DMX/DMXPixelMapping/Source/DMXPixelMappingRuntime/Private/Components/DMXPixelMappingFixtureGroupItemComponent.cpp b/Engine/Plugins/VirtualProduction/DMX/DMXPixelMapping/Source/DMXPixelMappingRuntime/Private/Components/DMXPixelMappingFixtureGroupItemComponent.cpp
--- a/Engine/Plugins/VirtualProduction/DMX/DMXPixelMapping/Source/DMXPixelMappingRuntime/Private/Components/DMXPixelMappingFixtureGroupItemComponent.cpp
+++ b/Engine/Plugins/VirtualProduction/DMX/DMXPixelMapping/Source/DMXPixelMappingRuntime/Private/Components/DMXPixelMappingFixtureGroupItemComponent.cpp
@@ -430,17 +430,20 @@ bool UDMXPixelMappingFixtureGroupItemComponent::IsOverParent() const
 	return false;
 }
 
-void UDMXPixelMappingFixtureGroupItemComponent::RenderWithInputAndSendDMX()
-{
-	RenderAndSendDMX();
-}
-
 TSharedRef<UE::DMXPixelMapping::Rendering::FPixelMapRenderElement> UDMXPixelMappingFixtureGroupItemComponent::GetOrCreatePixelMapRenderElement()
 {
 	UpdateRenderElement();
 	return PixelMapRenderElement.ToSharedRef();
 }
 
+PRAGMA_DISABLE_DEPRECATION_WARNINGS
+void UDMXPixelMappingFixtureGroupItemComponent::RenderWithInputAndSendDMX()
+{
+	// DEPRECATED 5.3
+	RenderAndSendDMX();
+}
+PRAGMA_ENABLE_DEPRECATION_WARNINGS
+
 bool UDMXPixelMappingFixtureGroupItemComponent::CanBeMovedTo(const UDMXPixelMappingBaseComponent* Component) const
 {
 	if (const UDMXPixelMappingFixtureGroupComponent* FixtureGroupComponent = Cast<UDMXPixelMappingFixtureGroupComponent>(Component))
@@ -467,7 +470,7 @@ void UDMXPixelMappingFixtureGroupItemComponent::UpdateRenderElement()
 	Parameters.UV = FVector2D(GetPosition().X / InputTextureWidth, GetPosition().Y / InputTextureHeight);
 	Parameters.UVSize = FVector2D(GetSize().X / InputTextureWidth, GetSize().Y / InputTextureHeight);
 	Parameters.UVCellSize = Parameters.UVSize / 2.f;
-	Parameters.CellBlendingQuality = static_cast<ECellBlendingQuality>(CellBlendingQuality);
+	Parameters.CellBlendingQuality = CellBlendingQuality;
 	Parameters.bStaticCalculateUV = true;
 
 	if (!PixelMapRenderElement.IsValid())
diff --git a/Engine/Plugins/VirtualProduction/DMX/DMXPixelMapping/Source/DMXPixelMappingRuntime/Private/Components/DMXPixelMappingMatrixCellComponent.cpp b/Engine/Plugins/VirtualProduction/DMX/DMXPixelMapping/Source/DMXPixelMappingRuntime/Private/Components/DMXPixelMappingMatrixCellComponent.cpp
--- a/Engine/Plugins/VirtualProduction/DMX/DMXPixelMapping/Source/DMXPixelMappingRuntime/Private/Components/DMXPixelMappingMatrixCellComponent.cpp
+++ b/Engine/Plugins/VirtualProduction/DMX/DMXPixelMapping/Source/DMXPixelMappingRuntime/Private/Components/DMXPixelMappingMatrixCellComponent.cpp
@@ -7,11 +7,11 @@
 #include "Components/DMXPixelMappingMatrixComponent.h"
 #include "Components/DMXPixelMappingRendererComponent.h"
 #include "DMXConversions.h"
+#include "DMXPixelMappingRenderElement.h"
 #include "DMXPixelMappingTypes.h"
 #include "Engine/Texture.h"
 #include "Library/DMXEntityFixturePatch.h"
 #include "Library/DMXEntityFixtureType.h"
-#include "PixelMapRenderElement.h"
 #include "TextureResource.h"
 
 #if WITH_EDITOR
@@ -299,7 +299,7 @@ void UDMXPixelMappingMatrixCellComponent::QueueDownsample()
 		UV,
 		UVSize,
 		UVCellSize,
-		CellBlendingQuality,
+		MatrixComponent->CellBlendingQuality,
 		bStaticCalculateUV
 	};
 
@@ -312,13 +312,6 @@ void UDMXPixelMappingMatrixCellComponent::SetCellCoordinate(FIntPoint InCellCoor
 	CellCoordinate = InCellCoordinate;
 }
 
-PRAGMA_DISABLE_DEPRECATION_WARNINGS
-void UDMXPixelMappingMatrixCellComponent::RenderWithInputAndSendDMX()
-{
-	RenderAndSendDMX();
-}
-PRAGMA_ENABLE_DEPRECATION_WARNINGS
-
 bool UDMXPixelMappingMatrixCellComponent::CanBeMovedTo(const UDMXPixelMappingBaseComponent* OtherComponent) const
 {
 	return OtherComponent && OtherComponent == GetParent();
@@ -334,6 +327,12 @@ void UDMXPixelMappingMatrixCellComponent::UpdateRenderElement()
 {
 	using namespace UE::DMXPixelMapping::Rendering;
 	
+	UDMXPixelMappingMatrixComponent* MatrixComponent = Cast<UDMXPixelMappingMatrixComponent>(GetParent());
+	if (!MatrixComponent)
+	{
+		return;
+	}
+
 	const UDMXPixelMappingRendererComponent* RendererComponent = GetRendererComponent();
 	const UTexture* InputTexture = RendererComponent ? RendererComponent->GetRenderedInputTexture() : nullptr;
 	const double InputTextureWidth = InputTexture ? InputTexture->GetSurfaceWidth() : 1.0;
@@ -343,7 +342,7 @@ void UDMXPixelMappingMatrixCellComponent::UpdateRenderElement()
 	Parameters.UV = FVector2D(GetPosition().X / InputTextureWidth, GetPosition().Y / InputTextureHeight);
 	Parameters.UVSize = FVector2D(GetSize().X / InputTextureWidth, GetSize().Y / InputTextureHeight);
 	Parameters.UVCellSize = Parameters.UVSize / 2.f;
-	Parameters.CellBlendingQuality = static_cast<ECellBlendingQuality>(CellBlendingQuality);
+	Parameters.CellBlendingQuality = MatrixComponent->CellBlendingQuality;
 	Parameters.bStaticCalculateUV = true;
 
 	if (!PixelMapRenderElement.IsValid())
diff --git a/Engine/Plugins/VirtualProduction/DMX/DMXPixelMapping/Source/DMXPixelMappingRuntime/Private/Components/DMXPixelMappingMatrixComponent.cpp b/Engine/Plugins/VirtualProduction/DMX/DMXPixelMapping/Source/DMXPixelMappingRuntime/Private/Components/DMXPixelMappingMatrixComponent.cpp
--- a/Engine/Plugins/VirtualProduction/DMX/DMXPixelMapping/Source/DMXPixelMappingRuntime/Private/Components/DMXPixelMappingMatrixComponent.cpp
+++ b/Engine/Plugins/VirtualProduction/DMX/DMXPixelMapping/Source/DMXPixelMappingRuntime/Private/Components/DMXPixelMappingMatrixComponent.cpp
@@ -440,6 +440,14 @@ void UDMXPixelMappingMatrixComponent::SetSize(const FVector2D& NewSize)
 	}
 }
 
+PRAGMA_DISABLE_DEPRECATION_WARNINGS
+void UDMXPixelMappingMatrixComponent::RenderWithInputAndSendDMX()
+{
+	// DEPRECATED 5.3
+	RenderAndSendDMX();
+}
+PRAGMA_ENABLE_DEPRECATION_WARNINGS
+
 void UDMXPixelMappingMatrixComponent::HandlePositionChanged()
 {
 #if WITH_EDITOR
diff --git a/Engine/Plugins/VirtualProduction/DMX/DMXPixelMapping/Source/DMXPixelMappingRuntime/Private/Components/DMXPixelMappingOutputComponent.cpp b/Engine/Plugins/VirtualProduction/DMX/DMXPixelMapping/Source/DMXPixelMappingRuntime/Private/Components/DMXPixelMappingOutputComponent.cpp
--- a/Engine/Plugins/VirtualProduction/DMX/DMXPixelMapping/Source/DMXPixelMappingRuntime/Private/Components/DMXPixelMappingOutputComponent.cpp
+++ b/Engine/Plugins/VirtualProduction/DMX/DMXPixelMapping/Source/DMXPixelMappingRuntime/Private/Components/DMXPixelMappingOutputComponent.cpp
@@ -13,8 +13,7 @@
 
 
 UDMXPixelMappingOutputComponent::UDMXPixelMappingOutputComponent()
-	: CellBlendingQuality(EDMXPixelBlendingQuality::Low)
-	, PositionX(0.f)
+	: PositionX(0.f)
 	, PositionY(0.f)
 	, SizeX(1.f)
 	, SizeY(1.f)
@@ -53,17 +52,6 @@ PRAGMA_DISABLE_DEPRECATION_WARNINGS
 			}, bSetVisibilityRecursive);
 	}
 PRAGMA_ENABLE_DEPRECATION_WARNINGS
-
-	// Apply cell belending quality
-	if (PropertyChangedEvent.GetPropertyName() == GET_MEMBER_NAME_CHECKED(UDMXPixelMappingOutputComponent, CellBlendingQuality))
-	{
-		// Propagonate to children
-		constexpr bool bSetCellBlendingQualityToChildsRecursive = true;
-		ForEachChildOfClass<UDMXPixelMappingOutputComponent>([this](UDMXPixelMappingOutputComponent* ChildComponent)
-			{
-				ChildComponent->CellBlendingQuality = CellBlendingQuality;
-			}, bSetCellBlendingQualityToChildsRecursive);
-	}
 }
 #endif // WITH_EDITOR
 
diff --git a/Engine/Plugins/VirtualProduction/DMX/DMXPixelMapping/Source/DMXPixelMappingRuntime/Private/Components/DMXPixelMappingScreenComponent.cpp b/Engine/Plugins/VirtualProduction/DMX/DMXPixelMapping/Source/DMXPixelMappingRuntime/Private/Components/DMXPixelMappingScreenComponent.cpp
--- a/Engine/Plugins/VirtualProduction/DMX/DMXPixelMapping/Source/DMXPixelMappingRuntime/Private/Components/DMXPixelMappingScreenComponent.cpp
+++ b/Engine/Plugins/VirtualProduction/DMX/DMXPixelMapping/Source/DMXPixelMappingRuntime/Private/Components/DMXPixelMappingScreenComponent.cpp
@@ -155,7 +155,7 @@ PRAGMA_ENABLE_DEPRECATION_WARNINGS
 
 const FName& UDMXPixelMappingScreenComponent::GetNamePrefix()
 {
-	static FName NamePrefix = TEXT("DMX Screen");
+	static FName NamePrefix = TEXT("DEPRECATED DMX Screen");
 	return NamePrefix;
 }
 
@@ -264,8 +264,6 @@ UDMXPixelMappingRendererComponent* UDMXPixelMappingScreenComponent::GetRendererC
 
 void UDMXPixelMappingScreenComponent::ResetDMX()
 {
-	//RenderElement->SetColor(FLinearColor::Black);
-
 	SendDMX();
 }
 
diff --git a/Engine/Plugins/VirtualProduction/DMX/DMXPixelMapping/Source/DMXPixelMappingRuntime/Public/Components/DMXPixelMappingFixtureGroupItemComponent.h b/Engine/Plugins/VirtualProduction/DMX/DMXPixelMapping/Source/DMXPixelMappingRuntime/Public/Components/DMXPixelMappingFixtureGroupItemComponent.h
--- a/Engine/Plugins/VirtualProduction/DMX/DMXPixelMapping/Source/DMXPixelMappingRuntime/Public/Components/DMXPixelMappingFixtureGroupItemComponent.h
+++ b/Engine/Plugins/VirtualProduction/DMX/DMXPixelMapping/Source/DMXPixelMappingRuntime/Public/Components/DMXPixelMappingFixtureGroupItemComponent.h
@@ -2,10 +2,10 @@
 
 #pragma once
 
-#include "Components/DMXPixelMappingOutputDMXComponent.h"
+#include "DMXPixelMappingOutputDMXComponent.h"
 
-#include "Components/DMXPixelMappingFixtureGroupComponent.h"
-#include "Components/DMXPixelMappingMatrixCellComponent.h"
+#include "DMXPixelMappingFixtureGroupComponent.h"
+#include "DMXPixelMappingMatrixCellComponent.h"
 #include "Library/DMXEntityReference.h"
 
 #include "Misc/Attribute.h"
@@ -69,25 +69,6 @@ public:
 	/** Check if a Component can be moved under another one (used for copy/move/duplicate) */
 	virtual bool CanBeMovedTo(const UDMXPixelMappingBaseComponent* Component) const override;
 
-	UPROPERTY(VisibleAnywhere, BlueprintReadOnly, Category = "Selected Patch")
-	FDMXEntityFixturePatchRef FixturePatchRef;
-
-	/** Sets which color space Pixel Mapping sends */
-	UPROPERTY(Transient, EditAnywhere, NoClear, Category = "Color Space", Meta = (DisplayPriority = 2, DisplayName = "Output Mode", ShowDisplayNames))
-	TSubclassOf<UDMXPixelMappingColorSpace> ColorSpaceClass;
-
-	/** The Color Space currently in use */
-	UPROPERTY(VisibleAnywhere, Instanced, Category = "Color Space")
-	TObjectPtr<UDMXPixelMappingColorSpace> ColorSpace;
-
-	/** Modulators applied to the output before sending DMX */
-	UPROPERTY(Transient, EditAnywhere, BlueprintReadOnly, Category = "Output Settings", Meta = (DisplayName = "Output Modulators"))
-	TArray<TSubclassOf<UDMXModulator>> ModulatorClasses;
-
-	/** Modulator instances applied to this component */
-	UPROPERTY()
-	TArray<TObjectPtr<UDMXModulator>> Modulators;
-
 	/** Index of the cell pixel in downsample target buffer */
 	int32 DownsamplePixelIndex;
 
@@ -117,8 +98,7 @@ public:
 	UE_DEPRECATED(5.3, "No longer in use. Please use UDMXPixelMappingPixelMapRenderer to render the pixel map")
 	virtual int32 GetDownsamplePixelIndex() const override { return DownsamplePixelIndex; }
 	//~ End UDMXPixelMappingOutputComponent implementation
-
-
+	
 	//~ Begin UDMXPixelMappingOutputDMXComponent implementation
 	PRAGMA_DISABLE_DEPRECATION_WARNINGS
 	UE_DEPRECATED(5.3, "Deprecated for performance reasons. Instead use 'Get DMX Pixel Mapping Renderer Component' and Render only once each tick.")
diff --git a/Engine/Plugins/VirtualProduction/DMX/DMXPixelMapping/Source/DMXPixelMappingRuntime/Public/Components/DMXPixelMappingMatrixCellComponent.h b/Engine/Plugins/VirtualProduction/DMX/DMXPixelMapping/Source/DMXPixelMappingRuntime/Public/Components/DMXPixelMappingMatrixCellComponent.h
--- a/Engine/Plugins/VirtualProduction/DMX/DMXPixelMapping/Source/DMXPixelMappingRuntime/Public/Components/DMXPixelMappingMatrixCellComponent.h
+++ b/Engine/Plugins/VirtualProduction/DMX/DMXPixelMapping/Source/DMXPixelMappingRuntime/Public/Components/DMXPixelMappingMatrixCellComponent.h
@@ -2,11 +2,11 @@
 
 #pragma once
 
-#include "Components/DMXPixelMappingOutputDMXComponent.h"
+#include "Components/DMXPixelMappingOutputComponent.h"
 #include "DMXAttribute.h"
+#include "DMXPixelMappingRenderElement.h"
 #include "DMXTypes.h"
 #include "Library/DMXEntityReference.h"
-#include "PixelMapRenderElement.h"
 
 #include "DMXPixelMappingMatrixCellComponent.generated.h"
 
@@ -23,7 +23,7 @@ class UTextureRenderTarget2D;
  */
 UCLASS()
 class DMXPIXELMAPPINGRUNTIME_API UDMXPixelMappingMatrixCellComponent
-	: public UDMXPixelMappingOutputDMXComponent
+	: public UDMXPixelMappingOutputComponent
 {
 	GENERATED_BODY()
 
@@ -64,12 +64,6 @@ public:
 	virtual void SetSize(const FVector2D& NewSize) override;
 	//~ End UDMXPixelMappingOutputComponent implementation
 
-	//~ Begin UDMXPixelMappingOutputDMXComponent implementation
-	PRAGMA_DISABLE_DEPRECATION_WARNINGS
-	virtual void RenderWithInputAndSendDMX() override;
-	PRAGMA_ENABLE_DEPRECATION_WARNINGS
-	//~ End UDMXPixelMappingOutputDMXComponent implementation
-
 	void SetCellCoordinate(FIntPoint InCellCoordinate);
 	const FIntPoint& GetCellCoordinate() { return CellCoordinate; }
 
diff --git a/Engine/Plugins/VirtualProduction/DMX/DMXPixelMapping/Source/DMXPixelMappingRuntime/Public/Components/DMXPixelMappingMatrixComponent.h b/Engine/Plugins/VirtualProduction/DMX/DMXPixelMapping/Source/DMXPixelMappingRuntime/Public/Components/DMXPixelMappingMatrixComponent.h
--- a/Engine/Plugins/VirtualProduction/DMX/DMXPixelMapping/Source/DMXPixelMappingRuntime/Public/Components/DMXPixelMappingMatrixComponent.h
+++ b/Engine/Plugins/VirtualProduction/DMX/DMXPixelMapping/Source/DMXPixelMappingRuntime/Public/Components/DMXPixelMappingMatrixComponent.h
@@ -2,6 +2,8 @@
 
 #pragma once
 
+#include "DMXPixelMappingOutputDMXComponent.h"
+
 #include "Components/DMXPixelMappingOutputComponent.h"
 #include "Library/DMXEntityReference.h"
 #include "Library/DMXEntityFixtureType.h"
@@ -18,7 +20,7 @@ class UDMXPixelMappingLayoutScript;
  */
 UCLASS()
 class DMXPIXELMAPPINGRUNTIME_API UDMXPixelMappingMatrixComponent
-	: public UDMXPixelMappingOutputComponent
+	: public UDMXPixelMappingOutputDMXComponent
 {
 	GENERATED_BODY()
 
@@ -69,6 +71,13 @@ public:
 	virtual void SetPosition(const FVector2D& NewPosition) override;
 	virtual void SetSize(const FVector2D& NewSize) override;
 	// ~End UDMXPixelMappingOutputComponent interface
+	
+	//~ Begin UDMXPixelMappingOutputDMXComponent implementation
+	PRAGMA_DISABLE_DEPRECATION_WARNINGS
+	UE_DEPRECATED(5.3, "Deprecated for performance reasons. Instead use 'Get DMX Pixel Mapping Renderer Component' and Render only once each tick.")
+	virtual void RenderWithInputAndSendDMX() override;
+	PRAGMA_ENABLE_DEPRECATION_WARNINGS
+	//~ End UDMXPixelMappingOutputDMXComponent implementation
 
 	/** Handles changes in position */
 	void HandlePositionChanged();
@@ -107,25 +116,6 @@ public:
 	FDMXEntityFixturePatchRef FixturePatchMatrixRef_DEPRECATED;
 #endif // WITH_EDITORONLY_DATA
 
-	UPROPERTY(VisibleAnywhere, BlueprintReadOnly, Category = "Selected Patch")
-	FDMXEntityFixturePatchRef FixturePatchRef;
-
-	/** Sets which color space Pixel Mapping sends */
-	UPROPERTY(Transient, EditAnywhere, NoClear, Category = "Color Space", Meta = (DisplayPriority = 2, DisplayName = "Output Mode", ShowDisplayNames))
-	TSubclassOf<UDMXPixelMappingColorSpace> ColorSpaceClass;
-
-	/** The Color Space currently in use */
-	UPROPERTY(VisibleAnywhere, Instanced, Category = "Color Space")
-	TObjectPtr<UDMXPixelMappingColorSpace> ColorSpace;
-
-	/** Modulators applied to the output before sending DMX */
-	UPROPERTY(EditAnywhere, BlueprintReadOnly, Transient, Category = "Output Settings", meta = (DisplayName = "Output Modulators"))
-	TArray<TSubclassOf<UDMXModulator>> ModulatorClasses;
-
-	/** The actual modulator instances */
-	UPROPERTY()
-	TArray<TObjectPtr<UDMXModulator>> Modulators;
-
 	UPROPERTY()
 	FIntPoint CoordinateGrid;
 
diff --git a/Engine/Plugins/VirtualProduction/DMX/DMXPixelMapping/Source/DMXPixelMappingRuntime/Public/Components/DMXPixelMappingOutputComponent.h b/Engine/Plugins/VirtualProduction/DMX/DMXPixelMapping/Source/DMXPixelMappingRuntime/Public/Components/DMXPixelMappingOutputComponent.h
--- a/Engine/Plugins/VirtualProduction/DMX/DMXPixelMapping/Source/DMXPixelMappingRuntime/Public/Components/DMXPixelMappingOutputComponent.h
+++ b/Engine/Plugins/VirtualProduction/DMX/DMXPixelMapping/Source/DMXPixelMappingRuntime/Public/Components/DMXPixelMappingOutputComponent.h
@@ -23,21 +23,6 @@ class FDMXPixelMappingComponentWidget;
 #endif
 
 
-
-/** Enum that defines the quality of how pixels are rendered */
-UENUM()
-enum class EDMXPixelBlendingQuality : uint8
-{
-	/** 1 sample */
-	Low,
-
-    /** 5 samples ( 2 x 2 with center) */
-    Medium,
-
-	/** 9 samples ( 3 x 3 ) */
-	High
-};
-
 /**
  * Base class for all Designer and configurable components
  */
@@ -192,21 +177,17 @@ public:
 	FORCEINLINE static FName GetSizeYPropertyName() { return GET_MEMBER_NAME_CHECKED(UDMXPixelMappingOutputComponent, SizeY); }
 #endif
 
-	/** The quality level to use when averaging colors during downsampling. */
-	UPROPERTY(EditAnywhere, BlueprintReadOnly, Category = "Pixel Settings")
-	EDMXPixelBlendingQuality CellBlendingQuality;
-
 private:
-	UPROPERTY(EditAnywhere, BlueprintReadOnly, Category = "Common Settings", Meta = (EditCondition = "!bLockInDesigner", AllowPrivateAccess = true))
+	UPROPERTY(EditAnywhere, BlueprintReadOnly, Category = "Transform", Meta = (EditCondition = "!bLockInDesigner", AllowPrivateAccess = true))
 	float PositionX;
 
-	UPROPERTY(EditAnywhere, BlueprintReadOnly, Category = "Common Settings", Meta = (EditCondition = "!bLockInDesigner", AllowPrivateAccess = true))
+	UPROPERTY(EditAnywhere, BlueprintReadOnly, Category = "Transform", Meta = (EditCondition = "!bLockInDesigner", AllowPrivateAccess = true))
 	float PositionY;
 
-	UPROPERTY(EditAnywhere, BlueprintReadOnly, Category = "Common Settings", Meta = (ClampMin = 1, UIMin = 1, EditCondition = "!bLockInDesigner", AllowPrivateAccess = true))
+	UPROPERTY(EditAnywhere, BlueprintReadOnly, Category = "Transform", Meta = (ClampMin = 1, UIMin = 1, EditCondition = "!bLockInDesigner", AllowPrivateAccess = true))
 	float SizeX;
 
-	UPROPERTY(EditAnywhere, BlueprintReadOnly, Category = "Common Settings", Meta = (ClampMin = 1, UIMin = 1, EditCondition = "!bLockInDesigner", AllowPrivateAccess = true))
+	UPROPERTY(EditAnywhere, BlueprintReadOnly, Category = "Transform", Meta = (ClampMin = 1, UIMin = 1, EditCondition = "!bLockInDesigner", AllowPrivateAccess = true))
 	float SizeY;
 
 protected:
diff --git a/Engine/Plugins/VirtualProduction/DMX/DMXPixelMapping/Source/DMXPixelMappingRuntime/Public/Components/DMXPixelMappingOutputDMXComponent.h b/Engine/Plugins/VirtualProduction/DMX/DMXPixelMapping/Source/DMXPixelMappingRuntime/Public/Components/DMXPixelMappingOutputDMXComponent.h
--- a/Engine/Plugins/VirtualProduction/DMX/DMXPixelMapping/Source/DMXPixelMappingRuntime/Public/Components/DMXPixelMappingOutputDMXComponent.h
+++ b/Engine/Plugins/VirtualProduction/DMX/DMXPixelMapping/Source/DMXPixelMappingRuntime/Public/Components/DMXPixelMappingOutputDMXComponent.h
@@ -2,12 +2,20 @@
 
 #pragma once
 
-#include "Components/DMXPixelMappingOutputComponent.h"
+#include "DMXPixelMappingOutputComponent.h"
+#include "Library/DMXEntityReference.h"
+
 #include "DMXPixelMappingOutputDMXComponent.generated.h"
 
+enum class EDMXPixelBlendingQuality : uint8;
+class UDMXModulator;
+class UDMXPixelMappingColorSpace;
+
 
 /**
- * Parent class for DMX sending components
+ * Base class for components that contain a fixture patch to send DMX.
+ * 
+ * For legacy reasons also used by deprecated DMXPixelMappingScreenComponent
  */
 UCLASS(Abstract)
 class DMXPIXELMAPPINGRUNTIME_API UDMXPixelMappingOutputDMXComponent
@@ -16,6 +24,29 @@ class DMXPIXELMAPPINGRUNTIME_API UDMXPixelMappingOutputDMXComponent
 	GENERATED_BODY()
 
 public:
+	/** The quality level to use when averaging colors during downsampling. */
+	UPROPERTY(EditAnywhere, BlueprintReadOnly, Category = "Quality")
+	EDMXPixelBlendingQuality CellBlendingQuality;
+
+	UPROPERTY(VisibleAnywhere, BlueprintReadOnly, Category = "Fixture Patch", Meta = (ShowOnlyInnerProperties))
+	FDMXEntityFixturePatchRef FixturePatchRef;
+
+	/** Sets which color space Pixel Mapping sends */
+	UPROPERTY(Transient, EditAnywhere, NoClear, Category = "Color Space", Meta = (DisplayPriority = 2, DisplayName = "Output Mode", ShowDisplayNames))
+	TSubclassOf<UDMXPixelMappingColorSpace> ColorSpaceClass;
+
+	/** The Color Space currently in use */
+	UPROPERTY(VisibleAnywhere, Instanced, Category = "Color Space")
+	TObjectPtr<UDMXPixelMappingColorSpace> ColorSpace;
+
+	/** Modulators applied to the output before sending DMX */
+	UPROPERTY(EditAnywhere, BlueprintReadOnly, Transient, Category = "Output Modulators", meta = (DisplayName = "Output Modulators"))
+	TArray<TSubclassOf<UDMXModulator>> ModulatorClasses;
+
+	/** The actual modulator instances */
+	UPROPERTY()
+	TArray<TObjectPtr<UDMXModulator>> Modulators;
+
 	/** Render input texture for downsample texture, donwsample and send DMX for this component */
 	UE_DEPRECATED(5.3, "Deprecated for performance reasons. Instead use 'Get DMX Pixel Mapping Renderer Component' and Render only once each tick.")
 	UFUNCTION(BlueprintCallable, Category = "DMX|PixelMapping", Meta = (DeprecatedFunction, DeprecationMessage = "Deprecated for performance reasons. Instead use 'Get DMX Pixel Mapping Renderer Component' and Render only once each tick"))
diff --git a/Engine/Plugins/VirtualProduction/DMX/DMXPixelMapping/Source/DMXPixelMappingRuntime/Public/Components/DMXPixelMappingScreenComponent.h b/Engine/Plugins/VirtualProduction/DMX/DMXPixelMapping/Source/DMXPixelMappingRuntime/Public/Components/DMXPixelMappingScreenComponent.h
--- a/Engine/Plugins/VirtualProduction/DMX/DMXPixelMapping/Source/DMXPixelMappingRuntime/Public/Components/DMXPixelMappingScreenComponent.h
+++ b/Engine/Plugins/VirtualProduction/DMX/DMXPixelMapping/Source/DMXPixelMappingRuntime/Public/Components/DMXPixelMappingScreenComponent.h
@@ -11,8 +11,8 @@
 
 #include "DMXPixelMappingScreenComponent.generated.h"
 
-
 enum class EDMXCellFormat : uint8;
+enum class EDMXPixelBlendingQuality : uint8;
 class FDMXPixelMappingComponentWidget;
 class FDMXOutputPort;
 class SDMXPixelMappingScreenComponentBox;
