DMX - Remove mistakenly submitted test code for deprecated screen component uprgade patth

diff --git a/Engine/Plugins/VirtualProduction/DMX/DMXPixelMapping/Source/DMXPixelMappingEditor/Private/Toolkits/DMXPixelMappingToolkit.cpp b/Engine/Plugins/VirtualProduction/DMX/DMXPixelMapping/Source/DMXPixelMappingEditor/Private/Toolkits/DMXPixelMappingToolkit.cpp
--- a/Engine/Plugins/VirtualProduction/DMX/DMXPixelMapping/Source/DMXPixelMappingEditor/Private/Toolkits/DMXPixelMappingToolkit.cpp
+++ b/Engine/Plugins/VirtualProduction/DMX/DMXPixelMapping/Source/DMXPixelMappingEditor/Private/Toolkits/DMXPixelMappingToolkit.cpp
@@ -590,9 +590,6 @@ void FDMXPixelMappingToolkit::InitializeInternal(const EToolkitMode::Type Mode,
 			SelectComponents(TSet<FDMXPixelMappingComponentReference>({ ComponentReference }));
 		}
 	}
-
-	FDMXPixelMappingComponentTemplate Template(UDMXPixelMappingScreenComponent::StaticClass());
-	GetActiveRendererComponent()->AddChild(Template.CreateComponent< UDMXPixelMappingScreenComponent>(DMXPixelMapping->GetRootComponent()));
 }
 
 TSharedRef<SDockTab> FDMXPixelMappingToolkit::SpawnTab_DMXLibraryView(const FSpawnTabArgs& Args)
