DMX - Add search and categories for the pixel mapping hierarchy view, better names, auto-generated thumbnails

diff --git a/Engine/Plugins/VirtualProduction/DMX/DMXPixelMapping/Source/DMXPixelMappingEditor/Private/Customizations/DMXPixelMappingColorSpaceDetails_RGBCMY.cpp b/Engine/Plugins/VirtualProduction/DMX/DMXPixelMapping/Source/DMXPixelMappingEditor/Private/Customizations/DMXPixelMappingColorSpaceDetails_RGBCMY.cpp
--- a/Engine/Plugins/VirtualProduction/DMX/DMXPixelMapping/Source/DMXPixelMappingEditor/Private/Customizations/DMXPixelMappingColorSpaceDetails_RGBCMY.cpp
+++ b/Engine/Plugins/VirtualProduction/DMX/DMXPixelMapping/Source/DMXPixelMappingEditor/Private/Customizations/DMXPixelMappingColorSpaceDetails_RGBCMY.cpp
@@ -10,7 +10,7 @@
 #include "PropertyHandle.h"
 #include "Widgets/Layout/SBox.h"
 #include "Widgets/Layout/SSeparator.h"
-#include "Widgets/Layout/SWrapBox.h"
+#include "Widgets/SBoxPanel.h"
 #include "Widgets/Text/STextBlock.h"
 
 
@@ -74,12 +74,10 @@ void FDMXPixelMappingColorSpaceDetails_RGBCMY::GenerateAttributeRow(IDetailLayou
 		.ValueContent()
 		.HAlign(HAlign_Fill)
 		[
-			SNew(SWrapBox)
-			.UseAllottedWidth(true)
-			.InnerSlotPadding(FVector2D(4.f, 0.f))
+			SNew(SHorizontalBox)
 
-			+ SWrapBox::Slot()
-			.HAlign(HAlign_Fill)
+			+ SHorizontalBox::Slot()
+			.AutoWidth()
 			[
 				SNew(SBox)
 				.MinDesiredWidth(120.f)
@@ -88,15 +86,15 @@ void FDMXPixelMappingColorSpaceDetails_RGBCMY::GenerateAttributeRow(IDetailLayou
 				]
 			]
 
-			+ SWrapBox::Slot()
-			.HAlign(HAlign_Left)
+			+ SHorizontalBox::Slot()
+			.AutoWidth()
 			.Padding(4.f, 0.f, 0.f, 0.f)
 			[
 				AttributeRowData.IsInvertColorHandle->CreatePropertyValueWidget()
 			]
 
-			+ SWrapBox::Slot()
-			.HAlign(HAlign_Left)
+			+ SHorizontalBox::Slot()
+			.AutoWidth()
 			.VAlign(VAlign_Center)
 			[
 				SNew(STextBlock)
@@ -104,8 +102,6 @@ void FDMXPixelMappingColorSpaceDetails_RGBCMY::GenerateAttributeRow(IDetailLayou
 				.Text(AttributeRowData.InvertColorLabel)
 			]
 		];
-
-	
 }
 
 FText FDMXPixelMappingColorSpaceDetails_RGBCMY::GetRedAttributeDisplayName() const
diff --git a/Engine/Plugins/VirtualProduction/DMX/DMXPixelMapping/Source/DMXPixelMappingEditor/Private/Customizations/DMXPixelMappingDetailCustomization_FixtureGroup.cpp b/Engine/Plugins/VirtualProduction/DMX/DMXPixelMapping/Source/DMXPixelMappingEditor/Private/Customizations/DMXPixelMappingDetailCustomization_FixtureGroup.cpp
--- a/Engine/Plugins/VirtualProduction/DMX/DMXPixelMapping/Source/DMXPixelMappingEditor/Private/Customizations/DMXPixelMappingDetailCustomization_FixtureGroup.cpp
+++ b/Engine/Plugins/VirtualProduction/DMX/DMXPixelMapping/Source/DMXPixelMappingEditor/Private/Customizations/DMXPixelMappingDetailCustomization_FixtureGroup.cpp
@@ -2,22 +2,14 @@
 
 #include "Customizations/DMXPixelMappingDetailCustomization_FixtureGroup.h"
 
-#include "DMXPixelMapping.h"
-#include "DMXPixelMappingLayoutSettings.h"
-#include "Components/DMXPixelMappingFixtureGroupComponent.h"
-#include "Components/DMXPixelMappingFixtureGroupItemComponent.h"
-#include "Components/DMXPixelMappingMatrixComponent.h"
-#include "Library/DMXEntityFixturePatch.h"
-#include "Library/DMXLibrary.h"
+#include "Algo/Find.h"
 
+#include "Components/DMXPixelMappingFixtureGroupComponent.h"
 #include "DetailLayoutBuilder.h"
-#include "DetailWidgetRow.h"
 #include "Editor.h"
 #include "IPropertyUtilities.h"
-#include "ScopedTransaction.h"
-#include "Algo/Find.h"
 #include "Misc/CoreDelegates.h"
-#include "Widgets/Input/SButton.h"
+#include "Settings/DMXPixelMappingEditorSettings.h"
 
 
 #define LOCTEXT_NAMESPACE "DMXPixelMappingDetailCustomization_FixtureGroup"
@@ -30,7 +22,6 @@ void FDMXPixelMappingDetailCustomization_FixtureGroup::CustomizeDetails(IDetailL
 	InDetailLayout.HideProperty(GET_MEMBER_NAME_CHECKED(UDMXPixelMappingFixtureGroupComponent, LayoutScript));
 
 	// Handle size changes
-	UpdateCachedScaleChildrenWithParent();
 	SizeXHandle = InDetailLayout.GetProperty(UDMXPixelMappingOutputComponent::GetSizeXPropertyName(), UDMXPixelMappingOutputComponent::StaticClass());
 	SizeXHandle->SetOnPropertyValuePreChange(FSimpleDelegate::CreateSP(this, &FDMXPixelMappingDetailCustomization_FixtureGroup::OnSizePropertyPreChange));
 	SizeXHandle->SetOnPropertyValueChangedWithData(TDelegate<void(const FPropertyChangedEvent&)>::CreateSP(this, &FDMXPixelMappingDetailCustomization_FixtureGroup::OnSizePropertyChanged));
@@ -39,34 +30,20 @@ void FDMXPixelMappingDetailCustomization_FixtureGroup::CustomizeDetails(IDetailL
 	SizeYHandle->SetOnPropertyValuePreChange(FSimpleDelegate::CreateSP(this, &FDMXPixelMappingDetailCustomization_FixtureGroup::OnSizePropertyPreChange));
 	SizeYHandle->SetOnPropertyValueChangedWithData(TDelegate<void(const FPropertyChangedEvent&)>::CreateSP(this, &FDMXPixelMappingDetailCustomization_FixtureGroup::OnSizePropertyChanged));
 
-	// Remember the group being edited
-	WeakFixtureGroupComponent = GetSelectedFixtureGroupComponent(InDetailLayout);
-
 	// Listen to component changes
-	UDMXPixelMappingBaseComponent::GetOnComponentAdded().AddSP(this, &FDMXPixelMappingDetailCustomization_FixtureGroup::OnComponentAdded);
-	UDMXPixelMappingBaseComponent::GetOnComponentRemoved().AddSP(this, &FDMXPixelMappingDetailCustomization_FixtureGroup::OnComponentRemoved);
+	UDMXPixelMappingBaseComponent::GetOnComponentAdded().AddSP(this, &FDMXPixelMappingDetailCustomization_FixtureGroup::OnComponentAddedOrRemoved);
 }
 
-void FDMXPixelMappingDetailCustomization_FixtureGroup::OnComponentAdded(UDMXPixelMapping* PixelMapping, UDMXPixelMappingBaseComponent* Component)
+void FDMXPixelMappingDetailCustomization_FixtureGroup::OnComponentAddedOrRemoved(UDMXPixelMapping* PixelMapping, UDMXPixelMappingBaseComponent* Component)
 {
-	if (!RequestForceRefreshHandle.IsValid())
-	{
-		RequestForceRefreshHandle = FCoreDelegates::OnEndFrame.AddSP(this, &FDMXPixelMappingDetailCustomization_FixtureGroup::ForceRefresh);
-	}
-}
-
-void FDMXPixelMappingDetailCustomization_FixtureGroup::OnComponentRemoved(UDMXPixelMapping* PixelMapping, UDMXPixelMappingBaseComponent* Component)
-{
-	if (!RequestForceRefreshHandle.IsValid())
-	{
-		RequestForceRefreshHandle = FCoreDelegates::OnEndFrame.AddSP(this, &FDMXPixelMappingDetailCustomization_FixtureGroup::ForceRefresh);
-	}
+	PropertyUtilities->RequestRefresh();
 }
 
 void FDMXPixelMappingDetailCustomization_FixtureGroup::OnSizePropertyPreChange()
 {
-	UpdateCachedScaleChildrenWithParent();
-	if (!bCachedScaleChildrenWithParent)
+	// Remember the sizes if children scale with parent to create the right transaction for interactive changes
+	const FDMXPixelMappingDesignerSettings& DesignerSettings = GetDefault<UDMXPixelMappingEditorSettings>()->DesignerSettings;
+	if (!DesignerSettings.bScaleChildrenWithParent)
 	{
 		return;
 	}
@@ -107,7 +84,8 @@ void FDMXPixelMappingDetailCustomization_FixtureGroup::OnSizePropertyChanged(con
 void FDMXPixelMappingDetailCustomization_FixtureGroup::HandleSizePropertyChanged()
 {
 	// Scale children if desired
-	if (!bCachedScaleChildrenWithParent || 
+	const FDMXPixelMappingDesignerSettings& DesignerSettings = GetDefault<UDMXPixelMappingEditorSettings>()->DesignerSettings;
+	if (!DesignerSettings.bScaleChildrenWithParent ||
 		PreEditChangeComponentToSizeMap.IsEmpty())
 	{
 		return;
@@ -163,54 +141,4 @@ void FDMXPixelMappingDetailCustomization_FixtureGroup::HandleSizePropertyChanged
 	}
 }
 
-void FDMXPixelMappingDetailCustomization_FixtureGroup::ForceRefresh()
-{
-	// Reset the handles so they won't fire any changes after refreshing
-	if (ensure(PropertyUtilities.IsValid()))
-	{
-		PropertyUtilities->ForceRefresh();
-	}
-
-	RequestForceRefreshHandle.Reset();
-}
-
-void FDMXPixelMappingDetailCustomization_FixtureGroup::UpdateCachedScaleChildrenWithParent()
-{
-	const UDMXPixelMappingLayoutSettings* LayoutSettings = GetDefault<UDMXPixelMappingLayoutSettings>();
-	if (LayoutSettings)
-	{
-		bCachedScaleChildrenWithParent = LayoutSettings->bScaleChildrenWithParent;
-	}
-}
-
-UDMXLibrary* FDMXPixelMappingDetailCustomization_FixtureGroup::GetSelectedDMXLibrary(UDMXPixelMappingFixtureGroupComponent* FixtureGroupComponent) const
-{
-	if (FixtureGroupComponent)
-	{
-		return FixtureGroupComponent->DMXLibrary;
-	}
-
-	return nullptr;
-}
-
-UDMXPixelMappingFixtureGroupComponent* FDMXPixelMappingDetailCustomization_FixtureGroup::GetSelectedFixtureGroupComponent(const IDetailLayoutBuilder& InDetailLayout) const
-{
-	const TArray<TWeakObjectPtr<UObject> >& SelectedSelectedObjects = InDetailLayout.GetSelectedObjects();
-	TArray<UDMXPixelMappingFixtureGroupComponent*> SelectedSelectedComponents;
-
-	for (TWeakObjectPtr<UObject> SelectedObject : SelectedSelectedObjects)
-	{
-		if (UDMXPixelMappingFixtureGroupComponent* SelectedComponent = Cast<UDMXPixelMappingFixtureGroupComponent>(SelectedObject.Get()))
-		{
-			SelectedSelectedComponents.Add(SelectedComponent);
-		}
-	}
-
-	// we only support 1 uobject editing for now
-	// we set singe UObject here SDMXPixelMappingDetailsView::OnSelectedComponenetChanged()
-	// and that is why we getting only one UObject from here TArray<TWeakObjectPtr<UObject>>& SDetailsView::GetSelectedObjects()
-	check(SelectedSelectedComponents.Num());
-	return SelectedSelectedComponents[0];
-}
-
 #undef LOCTEXT_NAMESPACE
diff --git a/Engine/Plugins/VirtualProduction/DMX/DMXPixelMapping/Source/DMXPixelMappingEditor/Private/Customizations/DMXPixelMappingDetailCustomization_FixtureGroup.h b/Engine/Plugins/VirtualProduction/DMX/DMXPixelMapping/Source/DMXPixelMappingEditor/Private/Customizations/DMXPixelMappingDetailCustomization_FixtureGroup.h
--- a/Engine/Plugins/VirtualProduction/DMX/DMXPixelMapping/Source/DMXPixelMappingEditor/Private/Customizations/DMXPixelMappingDetailCustomization_FixtureGroup.h
+++ b/Engine/Plugins/VirtualProduction/DMX/DMXPixelMapping/Source/DMXPixelMappingEditor/Private/Customizations/DMXPixelMappingDetailCustomization_FixtureGroup.h
@@ -8,7 +8,6 @@
 struct FPropertyChangedEvent;
 
 class FDMXPixelMappingToolkit;
-class UDMXLibrary;
 class UDMXPixelMapping;
 class UDMXPixelMappingBaseComponent;
 class UDMXPixelMappingFixtureGroupComponent;
@@ -36,11 +35,8 @@ public:
 	virtual void CustomizeDetails(IDetailLayoutBuilder& DetailLayout) override;
 
 private:
-	/** Called when a component was added */
-	void OnComponentAdded(UDMXPixelMapping* PixelMapping, UDMXPixelMappingBaseComponent* Component);
-
-	/** Called when a component was removed */
-	void OnComponentRemoved(UDMXPixelMapping* PixelMapping, UDMXPixelMappingBaseComponent* Component);
+	/** Called when a component was added or removed */
+	void OnComponentAddedOrRemoved(UDMXPixelMapping* PixelMapping, UDMXPixelMappingBaseComponent* Component);
 
 	/** Called before the SizeX property changed */
 	void OnSizePropertyPreChange();
@@ -51,38 +47,17 @@ private:
 	/** Handles the size property changed, useful to call on tick */
 	void HandleSizePropertyChanged();
 
-	/** Forces the detail layout to refresh */
-	void ForceRefresh();
-
-	/** Updates the bCachedScaleChildrenWithParent member from UDMXPixelMappingLayoutSettings */
-	void UpdateCachedScaleChildrenWithParent();
-
-	/** Helper that returns the library selected in for the group */
-	UDMXLibrary* GetSelectedDMXLibrary(UDMXPixelMappingFixtureGroupComponent* FixtureGroupComponent) const;
-
-	/** Returns the currently selected fixture group */
-	UDMXPixelMappingFixtureGroupComponent* GetSelectedFixtureGroupComponent(const IDetailLayoutBuilder& InDetailLayout) const;
-
 	/** Weak reference to the DMX editor */
 	TWeakPtr<FDMXPixelMappingToolkit> ToolkitWeakPtr;
 
-	/** The single fixture group component in use */
-	TWeakObjectPtr<UDMXPixelMappingFixtureGroupComponent> WeakFixtureGroupComponent;
-
 	/** SizeX before it got changed */
 	TMap<TWeakObjectPtr<UDMXPixelMappingFixtureGroupComponent>, FVector2D> PreEditChangeComponentToSizeMap;
 
-	/** Cached UDMXPixelMappingLayoutSettings::bScaleChildrenWithParen, to avoid repetitive reads during interactive changes */
-	bool bCachedScaleChildrenWithParent = false;
-
 	/** Property handle for the SizeX property */
 	TSharedPtr<IPropertyHandle> SizeXHandle;
 
 	/** Property handle for the SizeY property */
 	TSharedPtr<IPropertyHandle> SizeYHandle;
 
-	/** Delegate handle while being bound to the OnEndFrame event */
-	FDelegateHandle RequestForceRefreshHandle;
-
 	TSharedPtr<IPropertyUtilities> PropertyUtilities;
 };
diff --git a/Engine/Plugins/VirtualProduction/DMX/DMXPixelMapping/Source/DMXPixelMappingEditor/Private/Customizations/DMXPixelMappingDetailCustomization_FixtureGroupItem.cpp b/Engine/Plugins/VirtualProduction/DMX/DMXPixelMapping/Source/DMXPixelMappingEditor/Private/Customizations/DMXPixelMappingDetailCustomization_FixtureGroupItem.cpp
--- a/Engine/Plugins/VirtualProduction/DMX/DMXPixelMapping/Source/DMXPixelMappingEditor/Private/Customizations/DMXPixelMappingDetailCustomization_FixtureGroupItem.cpp
+++ b/Engine/Plugins/VirtualProduction/DMX/DMXPixelMapping/Source/DMXPixelMappingEditor/Private/Customizations/DMXPixelMappingDetailCustomization_FixtureGroupItem.cpp
@@ -119,8 +119,8 @@ void FDMXPixelMappingDetailCustomization_FixtureGroupItem::CreateModulatorDetail
 	IDetailCategoryBuilder& ModualtorsCategory = InDetailLayout.EditCategory("Modulators", LOCTEXT("DMXModulatorsCategory", "Modulators"), ECategoryPriority::Important);
 
 	TSharedPtr<IPropertyHandle> ModulatorClassesHandle = InDetailLayout.GetProperty(GET_MEMBER_NAME_CHECKED(UDMXPixelMappingFixtureGroupItemComponent, ModulatorClasses), UDMXPixelMappingFixtureGroupItemComponent::StaticClass());
-	ModulatorClassesHandle->SetOnPropertyValueChanged(FSimpleDelegate::CreateSP(this, &FDMXPixelMappingDetailCustomization_FixtureGroupItem::ForceRefresh));
-	ModulatorClassesHandle->SetOnChildPropertyValueChanged(FSimpleDelegate::CreateSP(this, &FDMXPixelMappingDetailCustomization_FixtureGroupItem::ForceRefresh));
+	ModulatorClassesHandle->SetOnPropertyValueChanged(FSimpleDelegate::CreateSP(this, &FDMXPixelMappingDetailCustomization_FixtureGroupItem::RequestRefresh));
+	ModulatorClassesHandle->SetOnChildPropertyValueChanged(FSimpleDelegate::CreateSP(this, &FDMXPixelMappingDetailCustomization_FixtureGroupItem::RequestRefresh));
 
 	ModualtorsCategory.AddProperty(ModulatorClassesHandle);
 
@@ -206,7 +206,7 @@ void FDMXPixelMappingDetailCustomization_FixtureGroupItem::CreateModulatorDetail
 	}
 }
 
-void FDMXPixelMappingDetailCustomization_FixtureGroupItem::ForceRefresh()
+void FDMXPixelMappingDetailCustomization_FixtureGroupItem::RequestRefresh()
 {
 	PropertyUtilities->ForceRefresh();
 }
diff --git a/Engine/Plugins/VirtualProduction/DMX/DMXPixelMapping/Source/DMXPixelMappingEditor/Private/Customizations/DMXPixelMappingDetailCustomization_FixtureGroupItem.h b/Engine/Plugins/VirtualProduction/DMX/DMXPixelMapping/Source/DMXPixelMappingEditor/Private/Customizations/DMXPixelMappingDetailCustomization_FixtureGroupItem.h
--- a/Engine/Plugins/VirtualProduction/DMX/DMXPixelMapping/Source/DMXPixelMappingEditor/Private/Customizations/DMXPixelMappingDetailCustomization_FixtureGroupItem.h
+++ b/Engine/Plugins/VirtualProduction/DMX/DMXPixelMapping/Source/DMXPixelMappingEditor/Private/Customizations/DMXPixelMappingDetailCustomization_FixtureGroupItem.h
@@ -52,8 +52,8 @@ private:
 	/** Creates Details for the Output Modulators */
 	void CreateModulatorDetails(IDetailLayoutBuilder& InDetailLayout);
 
-	/** Forces the layout to redraw */
-	void ForceRefresh();
+	/** Requests to refresh the details view */
+	void RequestRefresh();
 
 	/** Fixture Group Item Components that are being edited */
 	TArray<TWeakObjectPtr<UDMXPixelMappingFixtureGroupItemComponent>> FixtureGroupItemComponents;
diff --git a/Engine/Plugins/VirtualProduction/DMX/DMXPixelMapping/Source/DMXPixelMappingEditor/Private/Customizations/DMXPixelMappingDetailCustomization_Matrix.cpp b/Engine/Plugins/VirtualProduction/DMX/DMXPixelMapping/Source/DMXPixelMappingEditor/Private/Customizations/DMXPixelMappingDetailCustomization_Matrix.cpp
--- a/Engine/Plugins/VirtualProduction/DMX/DMXPixelMapping/Source/DMXPixelMappingEditor/Private/Customizations/DMXPixelMappingDetailCustomization_Matrix.cpp
+++ b/Engine/Plugins/VirtualProduction/DMX/DMXPixelMapping/Source/DMXPixelMappingEditor/Private/Customizations/DMXPixelMappingDetailCustomization_Matrix.cpp
@@ -114,8 +114,8 @@ void FDMXPixelMappingDetailCustomization_Matrix::CreateModulatorDetails(IDetailL
 	IDetailCategoryBuilder& ModualtorsCategory = InDetailLayout.EditCategory("Modulators", LOCTEXT("DMXModulatorsCategory", "Modulators"), ECategoryPriority::Important);
 
 	TSharedPtr<IPropertyHandle> ModulatorClassesHandle = InDetailLayout.GetProperty(GET_MEMBER_NAME_CHECKED(UDMXPixelMappingMatrixComponent, ModulatorClasses), UDMXPixelMappingMatrixComponent::StaticClass());
-	ModulatorClassesHandle->SetOnPropertyValueChanged(FSimpleDelegate::CreateSP(this, &FDMXPixelMappingDetailCustomization_Matrix::ForceRefresh));
-	ModulatorClassesHandle->SetOnChildPropertyValueChanged(FSimpleDelegate::CreateSP(this, &FDMXPixelMappingDetailCustomization_Matrix::ForceRefresh));
+	ModulatorClassesHandle->SetOnPropertyValueChanged(FSimpleDelegate::CreateSP(this, &FDMXPixelMappingDetailCustomization_Matrix::RequestRefresh));
+	ModulatorClassesHandle->SetOnChildPropertyValueChanged(FSimpleDelegate::CreateSP(this, &FDMXPixelMappingDetailCustomization_Matrix::RequestRefresh));
 
 	ModualtorsCategory.AddProperty(ModulatorClassesHandle);
 
@@ -200,9 +200,9 @@ void FDMXPixelMappingDetailCustomization_Matrix::CreateModulatorDetails(IDetailL
 	}
 }
 
-void FDMXPixelMappingDetailCustomization_Matrix::ForceRefresh()
+void FDMXPixelMappingDetailCustomization_Matrix::RequestRefresh()
 {
-	PropertyUtilities->ForceRefresh();
+	PropertyUtilities->RequestRefresh();
 }
 
 #undef LOCTEXT_NAMESPACE
diff --git a/Engine/Plugins/VirtualProduction/DMX/DMXPixelMapping/Source/DMXPixelMappingEditor/Private/Customizations/DMXPixelMappingDetailCustomization_Matrix.h b/Engine/Plugins/VirtualProduction/DMX/DMXPixelMapping/Source/DMXPixelMappingEditor/Private/Customizations/DMXPixelMappingDetailCustomization_Matrix.h
--- a/Engine/Plugins/VirtualProduction/DMX/DMXPixelMapping/Source/DMXPixelMappingEditor/Private/Customizations/DMXPixelMappingDetailCustomization_Matrix.h
+++ b/Engine/Plugins/VirtualProduction/DMX/DMXPixelMapping/Source/DMXPixelMappingEditor/Private/Customizations/DMXPixelMappingDetailCustomization_Matrix.h
@@ -50,8 +50,8 @@ private:
 	/** Creates Details for the Output Modulators */
 	void CreateModulatorDetails(IDetailLayoutBuilder& InDetailLayout);
 
-	/** Forces the layout to redraw */
-	void ForceRefresh();
+	/** Requests to refresh the details view */
+	void RequestRefresh();
 
 	/** Matrix Components that are being edited */
 	TArray<TWeakObjectPtr<UDMXPixelMappingMatrixComponent>> MatrixComponents;
diff --git a/Engine/Plugins/VirtualProduction/DMX/DMXPixelMapping/Source/DMXPixelMappingEditor/Private/Customizations/DMXPixelMappingDetailCustomization_Renderer.cpp b/Engine/Plugins/VirtualProduction/DMX/DMXPixelMapping/Source/DMXPixelMappingEditor/Private/Customizations/DMXPixelMappingDetailCustomization_Renderer.cpp
--- a/Engine/Plugins/VirtualProduction/DMX/DMXPixelMapping/Source/DMXPixelMappingEditor/Private/Customizations/DMXPixelMappingDetailCustomization_Renderer.cpp
+++ b/Engine/Plugins/VirtualProduction/DMX/DMXPixelMapping/Source/DMXPixelMappingEditor/Private/Customizations/DMXPixelMappingDetailCustomization_Renderer.cpp
@@ -98,7 +98,7 @@ void FDMXPixelMappingDetailCustomization_Renderer::PostUndo(bool bSuccess)
 	{
 		if (UDMXPixelMappingRendererComponent* RendererComponent = Cast<UDMXPixelMappingRendererComponent>(Object))
 		{
-			RendererComponent->InvalidatePreprocessRenderer();
+			RendererComponent->UpdatePreprocessRenderer();
 		}
 	}
 }
@@ -110,7 +110,7 @@ void FDMXPixelMappingDetailCustomization_Renderer::PostRedo(bool bSuccess)
 	{
 		if (UDMXPixelMappingRendererComponent* RendererComponent = Cast<UDMXPixelMappingRendererComponent>(Object))
 		{
-			RendererComponent->InvalidatePreprocessRenderer();
+			RendererComponent->UpdatePreprocessRenderer();
 		}
 	}
 }
diff --git a/Engine/Plugins/VirtualProduction/DMX/DMXPixelMapping/Source/DMXPixelMappingEditor/Private/DMXPixelMappingEditorCommands.cpp b/Engine/Plugins/VirtualProduction/DMX/DMXPixelMapping/Source/DMXPixelMappingEditor/Private/DMXPixelMappingEditorCommands.cpp
--- a/Engine/Plugins/VirtualProduction/DMX/DMXPixelMapping/Source/DMXPixelMappingEditor/Private/DMXPixelMappingEditorCommands.cpp
+++ b/Engine/Plugins/VirtualProduction/DMX/DMXPixelMapping/Source/DMXPixelMappingEditor/Private/DMXPixelMappingEditorCommands.cpp
@@ -21,11 +21,10 @@ void FDMXPixelMappingEditorCommands::RegisterCommands()
 {
 	UI_COMMAND(SaveThumbnailImage, "Thumbnail", "Generate Thumbnail", EUserInterfaceActionType::Button, FInputChord());
 
-	UI_COMMAND(AddMapping, "Add Mapping", "Add Mapping", EUserInterfaceActionType::Button, FInputChord());
+	UI_COMMAND(AddMapping, "Add Source", "Adds a new Source Texure, Material or User Widget to the Pixel Map asset.", EUserInterfaceActionType::Button, FInputChord());
 
-	UI_COMMAND(PlayDMX, "Play DMX", "Play DMX.", EUserInterfaceActionType::Button, FInputChord());
-	UI_COMMAND(StopPlayingDMX, "Stop Playing DMX", "Stop Playing DMX.", EUserInterfaceActionType::Button, FInputChord());
-	UI_COMMAND(TogglePlayDMXAll, "Toggle Play DMX All", "Toggles playing all render components and all their childs", EUserInterfaceActionType::ToggleButton, FInputChord());
+	UI_COMMAND(PlayDMX, "Play DMX", "Starts to send DMX", EUserInterfaceActionType::Button, FInputChord());
+	UI_COMMAND(StopPlayingDMX, "Stop Playing DMX", "Stops playing DMX.", EUserInterfaceActionType::Button, FInputChord());
 
 	UI_COMMAND(SizeComponentToTexture, "Size Component to Texture", "Sizes the selected group to Texture.", EUserInterfaceActionType::Button, FInputChord(EModifierKey::Shift, EKeys::S));
 	
diff --git a/Engine/Plugins/VirtualProduction/DMX/DMXPixelMapping/Source/DMXPixelMappingEditor/Private/DMXPixelMappingEditorCommands.h b/Engine/Plugins/VirtualProduction/DMX/DMXPixelMapping/Source/DMXPixelMappingEditor/Private/DMXPixelMappingEditorCommands.h
--- a/Engine/Plugins/VirtualProduction/DMX/DMXPixelMapping/Source/DMXPixelMappingEditor/Private/DMXPixelMappingEditorCommands.h
+++ b/Engine/Plugins/VirtualProduction/DMX/DMXPixelMapping/Source/DMXPixelMappingEditor/Private/DMXPixelMappingEditorCommands.h
@@ -18,7 +18,6 @@ public:
 	TSharedPtr<FUICommandInfo> AddMapping;
 	TSharedPtr<FUICommandInfo> PlayDMX;
 	TSharedPtr<FUICommandInfo> StopPlayingDMX;
-	TSharedPtr<FUICommandInfo> TogglePlayDMXAll;
 
 	// Layout related
 	TSharedPtr<FUICommandInfo> SizeComponentToTexture;
diff --git a/Engine/Plugins/VirtualProduction/DMX/DMXPixelMapping/Source/DMXPixelMappingEditor/Private/DMXPixelMappingEditorCommon.h b/Engine/Plugins/VirtualProduction/DMX/DMXPixelMapping/Source/DMXPixelMappingEditor/Private/DMXPixelMappingEditorCommon.h
--- a/Engine/Plugins/VirtualProduction/DMX/DMXPixelMapping/Source/DMXPixelMappingEditor/Private/DMXPixelMappingEditorCommon.h
+++ b/Engine/Plugins/VirtualProduction/DMX/DMXPixelMapping/Source/DMXPixelMappingEditor/Private/DMXPixelMappingEditorCommon.h
@@ -9,14 +9,14 @@ DECLARE_LOG_CATEGORY_EXTERN(LogDMXPixelMappingEditor, All, All);
 class UDMXPixelMappingBaseComponent;
 
 class FDMXPixelMappingToolkit;
-class FDMXPixelMappingHierarchyItemWidgetModel;
+class FDMXPixelMappingHierarchyItem;
 class FDMXPixelMappingPaletteWidgetViewModel;
 class FDMXPixelMappingComponentTemplate;
 
 using FDMXPixelMappingToolkitPtr = TSharedPtr<FDMXPixelMappingToolkit>;
 using FDMXPixelMappingToolkitWeakPtr = TWeakPtr<FDMXPixelMappingToolkit>;
-using FDMXPixelMappingHierarchyItemWidgetModelPtr = TSharedPtr<FDMXPixelMappingHierarchyItemWidgetModel>;
-using FDMXPixelMappingHierarchyWidgetModelWeakPtr = TWeakPtr<FDMXPixelMappingHierarchyItemWidgetModel>;
+using FDMXPixelMappingHierarchyItemWidgetModelPtr = TSharedPtr<FDMXPixelMappingHierarchyItem>;
+using FDMXPixelMappingHierarchyWidgetModelWeakPtr = TWeakPtr<FDMXPixelMappingHierarchyItem>;
 using FDMXPixelMappingHierarchyItemWidgetModelArr = TArray<FDMXPixelMappingHierarchyItemWidgetModelPtr>;
 using FDMXPixelMappingPreviewWidgetViewModelPtr = TSharedPtr<FDMXPixelMappingPaletteWidgetViewModel>;
 using FDMXPixelMappingPreviewWidgetViewModelArray = TArray<FDMXPixelMappingPreviewWidgetViewModelPtr>;
diff --git a/Engine/Plugins/VirtualProduction/DMX/DMXPixelMapping/Source/DMXPixelMappingEditor/Private/DMXPixelMappingEditorUtils.cpp b/Engine/Plugins/VirtualProduction/DMX/DMXPixelMapping/Source/DMXPixelMappingEditor/Private/DMXPixelMappingEditorUtils.cpp
--- a/Engine/Plugins/VirtualProduction/DMX/DMXPixelMapping/Source/DMXPixelMappingEditor/Private/DMXPixelMappingEditorUtils.cpp
+++ b/Engine/Plugins/VirtualProduction/DMX/DMXPixelMapping/Source/DMXPixelMappingEditor/Private/DMXPixelMappingEditorUtils.cpp
@@ -18,80 +18,6 @@
 #define LOCTEXT_NAMESPACE "FDMXPixelMappingEditorUtils"
 
 
-bool FDMXPixelMappingEditorUtils::VerifyComponentRename(TSharedRef<FDMXPixelMappingToolkit> InToolkit, const FDMXPixelMappingComponentReference& InComponent, const FText& NewName, FText& OutErrorMessage)
-{
-	if (!InComponent.IsValid())
-	{
-		OutErrorMessage = LOCTEXT("InvalidComponentReference", "Invalid Component Reference");
-		return false;
-	}
-
-	if (NewName.IsEmptyOrWhitespace())
-	{
-		OutErrorMessage = LOCTEXT("EmptyComponentName", "Empty Component Name");
-		return false;
-	}
-
-	const FString& NewNameString = NewName.ToString();
-
-	if (NewNameString.Len() >= NAME_SIZE)
-	{
-		OutErrorMessage = LOCTEXT("ComponentNameTooLong", "Component Name is Too Long");
-		return false;
-	}
-
-	UDMXPixelMappingBaseComponent* ComponentToRename = InComponent.GetComponent();
-	if (!ComponentToRename)
-	{
-		// In certain situations, the template might be lost due to mid recompile with focus lost on the rename box at
-		// during a strange moment.
-		return false;
-	}
-
-	// Slug the new name down to a valid object name
-	const FName NewNameSlug = MakeObjectNameFromDisplayLabel(NewNameString, ComponentToRename->GetFName());
-
-	UDMXPixelMapping* DMXPixelMapping = InToolkit->GetDMXPixelMapping();
-	UDMXPixelMappingBaseComponent* ExistingComponent = DMXPixelMapping->FindComponent(NewNameSlug);
-
-	if (ExistingComponent != nullptr)
-	{
-		if (ComponentToRename != ExistingComponent)
-		{
-			OutErrorMessage = LOCTEXT("ExistingComponentName", "Existing Component Name");
-			return false;
-		}
-	}
-	else
-	{
-		// check for redirectors too
-		if (FindObject<UObject>(ComponentToRename->GetOuter(), *NewNameSlug.ToString()))
-		{
-			OutErrorMessage = LOCTEXT("ExistingOldComponentName", "Existing Old Component Name");
-			return false;
-		}
-	}
-
-	return true;
-}
-
-void FDMXPixelMappingEditorUtils::RenameComponent(TSharedRef<FDMXPixelMappingToolkit> InToolkit, const FName& OldObjectName, const FString& NewDisplayName)
-{
-	UDMXPixelMapping* DMXPixelMapping = InToolkit->GetDMXPixelMapping();
-	check(DMXPixelMapping);
-
-	UDMXPixelMappingBaseComponent* ComponentToRename = DMXPixelMapping->FindComponent(OldObjectName);
-	check(ComponentToRename);
-
-	// Get the new FName slug from the given display name
-	const FName NewFName = MakeObjectNameFromDisplayLabel(NewDisplayName, ComponentToRename->GetFName());
-	const FString NewNameStr = NewFName.ToString();
-
-	ComponentToRename->Rename(*NewNameStr);
-
-	InToolkit->OnComponentRenamed(ComponentToRename);
-}
-
 UDMXPixelMappingRendererComponent* FDMXPixelMappingEditorUtils::AddRenderer(UDMXPixelMapping* InPixelMapping)
 {
 	if (InPixelMapping == nullptr)
@@ -122,21 +48,6 @@ UDMXPixelMappingRendererComponent* FDMXPixelMappingEditorUtils::AddRenderer(UDMX
 	return Component;
 }
 
-void FDMXPixelMappingEditorUtils::CreateComponentContextMenu(FMenuBuilder& MenuBuilder, TSharedRef<FDMXPixelMappingToolkit> InToolkit)
-{
-	MenuBuilder.BeginSection("Edit", LOCTEXT("Edit", "Edit"));
-	{
-		MenuBuilder.AddMenuEntry(FGenericCommands::Get().Rename);
-		MenuBuilder.AddMenuEntry(FGenericCommands::Get().Cut);
-		MenuBuilder.AddMenuEntry(FGenericCommands::Get().Copy);
-		MenuBuilder.AddMenuEntry(FGenericCommands::Get().Paste);
-		MenuBuilder.AddMenuEntry(FGenericCommands::Get().Duplicate);
-		MenuBuilder.AddMenuEntry(FGenericCommands::Get().Delete);
-	}
-	MenuBuilder.EndSection();
-}
-
-
 bool FDMXPixelMappingEditorUtils::GetArrangedWidget(TSharedRef<SWidget> InWidget, FArrangedWidget& OutArrangedWidget)
 {
 	TSharedPtr<SWindow> WidgetWindow = FSlateApplication::Get().FindWidgetWindow(InWidget);
diff --git a/Engine/Plugins/VirtualProduction/DMX/DMXPixelMapping/Source/DMXPixelMappingEditor/Private/DMXPixelMappingEditorUtils.h b/Engine/Plugins/VirtualProduction/DMX/DMXPixelMapping/Source/DMXPixelMappingEditor/Private/DMXPixelMappingEditorUtils.h
--- a/Engine/Plugins/VirtualProduction/DMX/DMXPixelMapping/Source/DMXPixelMappingEditor/Private/DMXPixelMappingEditorUtils.h
+++ b/Engine/Plugins/VirtualProduction/DMX/DMXPixelMapping/Source/DMXPixelMappingEditor/Private/DMXPixelMappingEditorUtils.h
@@ -24,27 +24,6 @@ class SWidget;
 class FDMXPixelMappingEditorUtils
 {
 public:
-	/**
-	 * Check ability to rename the component.
-	 *
-	 * @param InToolkit				Pixel Mappint editor toolkit
-	 * @param InComponent			Component reference to rename
-	 * @param NewName				New name to check
-	 * @param OutErrorMessage		Output parameter in case of failed reaming
-	 *
-	 * @return true if the component can be renamed
-	 */
-	static bool VerifyComponentRename(TSharedRef<FDMXPixelMappingToolkit> InToolkit, const FDMXPixelMappingComponentReference& InComponent, const FText& NewName, FText& OutErrorMessage);
-
-	/**
-	 * Rename Pixel Maping component.
-	 *
-	 * @param InToolkit				Pixel Mappint editor toolkit
-	 * @param NewName				Old name
-	 * @param NewName				New name
-	 */
-	static void RenameComponent(TSharedRef<FDMXPixelMappingToolkit> InToolkit, const FName& OldObjectName, const FString& NewDisplayName);
-
 	/**
 	 * Add renderer to Pixel Mapping object
 	 *
@@ -54,14 +33,6 @@ public:
 	 */
 	static UDMXPixelMappingRendererComponent* AddRenderer(UDMXPixelMapping* InPixelMapping);
 
-	/**
-	 * Create components commands menu
-	 *
-	 * @param MenuBuilder			Vertical menu builder
-	 * @param InToolkit				Pixel Mappint editor toolkit
-	 */
-	static void CreateComponentContextMenu(FMenuBuilder& MenuBuilder, TSharedRef<FDMXPixelMappingToolkit> InToolkit);
-
 	/**
 	 * Returns the arranged widget of a widget.
 	 *
diff --git a/Engine/Plugins/VirtualProduction/DMX/DMXPixelMapping/Source/DMXPixelMappingEditor/Private/DMXPixelMappingLayoutSettings.cpp b/Engine/Plugins/VirtualProduction/DMX/DMXPixelMapping/Source/DMXPixelMappingEditor/Private/DMXPixelMappingLayoutSettings.cpp
deleted file mode 100644
diff --git a/Engine/Plugins/VirtualProduction/DMX/DMXPixelMapping/Source/DMXPixelMappingEditor/Private/DMXPixelMappingLayoutSettings.h b/Engine/Plugins/VirtualProduction/DMX/DMXPixelMapping/Source/DMXPixelMappingEditor/Private/DMXPixelMappingLayoutSettings.h
deleted file mode 100644
diff --git a/Engine/Plugins/VirtualProduction/DMX/DMXPixelMapping/Source/DMXPixelMappingEditor/Private/DMXPixelMappingToolbar.cpp b/Engine/Plugins/VirtualProduction/DMX/DMXPixelMapping/Source/DMXPixelMappingEditor/Private/DMXPixelMappingToolbar.cpp
--- a/Engine/Plugins/VirtualProduction/DMX/DMXPixelMapping/Source/DMXPixelMappingEditor/Private/DMXPixelMappingToolbar.cpp
+++ b/Engine/Plugins/VirtualProduction/DMX/DMXPixelMapping/Source/DMXPixelMappingEditor/Private/DMXPixelMappingToolbar.cpp
@@ -29,53 +29,32 @@ void FDMXPixelMappingToolbar::BuildToolbar(TSharedPtr<FExtender> Extender)
 
 void FDMXPixelMappingToolbar::BuildToolbarCallback(FToolBarBuilder& ToolbarBuilder)
 {
-	ToolbarBuilder.BeginSection("Thumbnail");
-	{
-		ToolbarBuilder.AddToolBarButton(FDMXPixelMappingEditorCommands::Get().SaveThumbnailImage, NAME_None,
-			LOCTEXT("GenerateThumbnail", "Thumbnail"),
-			LOCTEXT("GenerateThumbnailTooltip", "Generate a thumbnail image."),
-			FSlateIcon(FAppStyle::GetAppStyleSetName(), "Cascade.SaveThumbnailImage"));
-	}
-	ToolbarBuilder.EndSection();
-
-	TSharedPtr<FDMXPixelMappingToolkit> Toolkit = ToolkitWeakPtr.Pin();
-	check(Toolkit.IsValid());
-
 	ToolbarBuilder.BeginSection("Renderers");
 	{
 		ToolbarBuilder.AddToolBarButton(FDMXPixelMappingEditorCommands::Get().AddMapping,
-			NAME_None, TAttribute<FText>(), TAttribute<FText>(),
+			NAME_None, 
+			TAttribute<FText>(), 
+			TAttribute<FText>(),
 			FSlateIcon(FDMXPixelMappingEditorStyle::GetStyleSetName(), "DMXPixelMappingEditor.AddMapping"),
-			FName(TEXT("Add Mapping")));
+			FName(TEXT("Add Source")));
 	}
 	ToolbarBuilder.EndSection();
 
 	ToolbarBuilder.BeginSection("PlayAndStopDMX");
 	{
 		ToolbarBuilder.AddToolBarButton(FDMXPixelMappingEditorCommands::Get().PlayDMX,
-			NAME_None, TAttribute<FText>(), TAttribute<FText>(),
+			NAME_None, 
+			TAttribute<FText>(), 
+			TAttribute<FText>(),
 			FSlateIcon(FDMXPixelMappingEditorStyle::GetStyleSetName(), "DMXPixelMappingEditor.PlayDMX"),
 			FName(TEXT("Play DMX")));
 
 		ToolbarBuilder.AddToolBarButton(FDMXPixelMappingEditorCommands::Get().StopPlayingDMX,
-			NAME_None, TAttribute<FText>(), TAttribute<FText>(),
+			NAME_None, 
+			TAttribute<FText>(), 
+			TAttribute<FText>(),
 			FSlateIcon(FDMXPixelMappingEditorStyle::GetStyleSetName(), "DMXPixelMappingEditor.StopPlayingDMX"),
 			FName(TEXT("Stop Playing DMX")));
-
-		FUIAction PlayDMXOptionsAction(
-			FExecuteAction(),
-			FCanExecuteAction(),
-			FIsActionChecked(),
-			FIsActionButtonVisible::CreateLambda([this] { return !ToolkitWeakPtr.Pin()->IsPlayingDMX(); })
-			);
-
-		ToolbarBuilder.AddComboButton(
-			PlayDMXOptionsAction,
-			FOnGetContent::CreateSP(this, &FDMXPixelMappingToolbar::GeneratesPlayOptionsWidget),
-			LOCTEXT("PlayDMXOptions", "Play DMX Options"),
-			LOCTEXT("PlayDMXOptions", "Play DMX Options"),
-			FSlateIcon(),
-			true);
 	}
 	ToolbarBuilder.EndSection();
 
@@ -92,25 +71,6 @@ void FDMXPixelMappingToolbar::BuildToolbarCallback(FToolBarBuilder& ToolbarBuild
 	ToolbarBuilder.EndSection();
 }
 
-TSharedRef<SWidget> FDMXPixelMappingToolbar::GeneratesPlayOptionsWidget()
-{
-	TSharedPtr<FDMXPixelMappingToolkit> Toolkit = ToolkitWeakPtr.Pin();
-	check(Toolkit.IsValid());
-
-	FMenuBuilder MenuBuilder(true, Toolkit->GetToolkitCommands());
-
-	if (!Toolkit->IsPlayingDMX())
-	{
-		MenuBuilder.BeginSection("bTogglePlayDMXAll");
-		{
-			MenuBuilder.AddMenuEntry(FDMXPixelMappingEditorCommands::Get().TogglePlayDMXAll);
-		}
-		MenuBuilder.EndSection();
-	}
-
-	return MenuBuilder.MakeWidget();
-}
-
 TSharedRef<SWidget> FDMXPixelMappingToolbar::GenerateLayoutMenu()
 {
 	TSharedPtr<FDMXPixelMappingToolkit> Toolkit = ToolkitWeakPtr.Pin();
@@ -125,17 +85,26 @@ TSharedRef<SWidget> FDMXPixelMappingToolbar::GenerateLayoutMenu()
 	}
 	MenuBuilder.EndSection();
 
-	MenuBuilder.BeginSection("LayoutSettings", LOCTEXT("LayoutSettingsSection", "Settings"));
+	MenuBuilder.BeginSection("EditorSettings", LOCTEXT("EditorSettingsSection", "Editor Settings"));
 	{
 		MenuBuilder.AddMenuEntry(FDMXPixelMappingEditorCommands::Get().ToggleScaleChildrenWithParent);
 		MenuBuilder.AddMenuEntry(FDMXPixelMappingEditorCommands::Get().ToggleAlwaysSelectGroup);
-		MenuBuilder.AddMenuEntry(FDMXPixelMappingEditorCommands::Get().ToggleApplyLayoutScriptWhenLoaded);
+	}
+
+	MenuBuilder.BeginSection("DisplaySettings", LOCTEXT("DisplaySettingsSection", "Display Settings"));
+	{
 		MenuBuilder.AddMenuEntry(FDMXPixelMappingEditorCommands::Get().ToggleShowComponentNames);
 		MenuBuilder.AddMenuEntry(FDMXPixelMappingEditorCommands::Get().ToggleShowPatchInfo);
+		MenuBuilder.AddMenuEntry(FDMXPixelMappingEditorCommands::Get().ToggleShowMatrixCells);
 		MenuBuilder.AddMenuEntry(FDMXPixelMappingEditorCommands::Get().ToggleShowCellIDs);
 	}
 	MenuBuilder.EndSection();
 
+	MenuBuilder.BeginSection("EditorSettings", LOCTEXT("LayoutSettingsSection", "Layout Settings"));
+	{
+		MenuBuilder.AddMenuEntry(FDMXPixelMappingEditorCommands::Get().ToggleApplyLayoutScriptWhenLoaded);
+	}
+
 	return MenuBuilder.MakeWidget();
 }
 
diff --git a/Engine/Plugins/VirtualProduction/DMX/DMXPixelMapping/Source/DMXPixelMappingEditor/Private/Settings/DMXPixelMappingEditorSettings.cpp b/Engine/Plugins/VirtualProduction/DMX/DMXPixelMapping/Source/DMXPixelMappingEditor/Private/Settings/DMXPixelMappingEditorSettings.cpp
new file mode 100644
--- /dev/null
+++ b/Engine/Plugins/VirtualProduction/DMX/DMXPixelMapping/Source/DMXPixelMappingEditor/Private/Settings/DMXPixelMappingEditorSettings.cpp
@@ -0,0 +1,6 @@
+// Copyright Epic Games, Inc. All Rights Reserved.
+
+#include "DMXPixelMappingEditorSettings.h"
+
+
+FSimpleMulticastDelegate UDMXPixelMappingEditorSettings::OnEditorSettingsChanged;
diff --git a/Engine/Plugins/VirtualProduction/DMX/DMXPixelMapping/Source/DMXPixelMappingEditor/Private/Settings/DMXPixelMappingEditorSettings.h b/Engine/Plugins/VirtualProduction/DMX/DMXPixelMapping/Source/DMXPixelMappingEditor/Private/Settings/DMXPixelMappingEditorSettings.h
new file mode 100644
--- /dev/null
+++ b/Engine/Plugins/VirtualProduction/DMX/DMXPixelMapping/Source/DMXPixelMappingEditor/Private/Settings/DMXPixelMappingEditorSettings.h
@@ -0,0 +1,89 @@
+// Copyright Epic Games, Inc. All Rights Reserved.
+
+#pragma once
+
+#include "DMXPixelMappingEditorSettings.generated.h"
+
+struct FPropertyChangedEvent;
+
+
+USTRUCT()
+struct FDMXPixelMappingHierarchySettings
+{
+	GENERATED_BODY()
+
+	/** If true, shows the editor color column */
+	UPROPERTY()
+	bool bShowEditorColorColumn = false;
+
+	/** If true, shows the fixture ID column */
+	UPROPERTY()
+	bool bShowFixtureIDColumn = true;
+
+	/** If true, shows the patch column */
+	UPROPERTY()
+	bool bShowPatchColumn = true;
+
+	/** The sort by column Id. Can be none or a valid column id */
+	UPROPERTY()
+	FName SortByColumnId;
+
+	/** If true, the current sort mode is ascending */
+	UPROPERTY()
+	bool bSortAscending = true;
+};
+
+USTRUCT()
+struct FDMXPixelMappingDesignerSettings
+{
+	GENERATED_BODY()
+
+	/** If true, scales children when the parent component is resized */
+	UPROPERTY()
+	bool bScaleChildrenWithParent = true;
+
+	/** If true, selects parent when a child is clicked */
+	UPROPERTY()
+	bool bAlwaysSelectGroup = false;
+
+	/** If true, applies layout scripts instantly when they are loaded */
+	UPROPERTY()
+	bool bApplyLayoutScriptWhenLoaded = true;
+
+	/**  If true, shows a widget for each cell. It is recommended that this is turned off when pixel mapping large quantities of fixtures. */
+	UPROPERTY()
+	bool bShowMatrixCells = true;
+
+	/** If true, shows the name of the Fixture Patch where applicable. It is recommended that this is turned off when pixel mapping large quantities of fixtures. */
+	UPROPERTY()
+	bool bShowComponentNames = true;
+
+	/** If true, shows Fixture info bout the Fixture Patch where applicable. It is recommended that this is turned off when pixel mapping large quantities of fixtures. */
+	UPROPERTY()
+	bool bShowPatchInfo = false;
+
+	/** If true, shows Cell IDs where applicable. It is recommended that this is turned off when pixel mapping large quantities of fixtures. */
+	UPROPERTY()
+	bool bShowCellIDs = false;
+};
+
+
+/** Layout options for the Pixel Mapping editor */
+UCLASS(Config = DMXPixelMappingEditor, DefaultConfig)
+class UDMXPixelMappingEditorSettings
+	: public UObject
+{
+	GENERATED_BODY()
+
+public:
+	/** Settings for the hierarchy tab */
+	UPROPERTY(Config)
+	FDMXPixelMappingHierarchySettings HierarchySettings;
+	
+	/** Settings for the designer tab */
+	UPROPERTY(Config)
+	FDMXPixelMappingDesignerSettings DesignerSettings;
+
+	/** Raised by the pixel mapping toolkit when settings changed */
+	static FSimpleMulticastDelegate OnEditorSettingsChanged;
+};
diff --git a/Engine/Plugins/VirtualProduction/DMX/DMXPixelMapping/Source/DMXPixelMappingEditor/Private/Toolkits/DMXPixelMappingToolkit.cpp b/Engine/Plugins/VirtualProduction/DMX/DMXPixelMapping/Source/DMXPixelMappingEditor/Private/Toolkits/DMXPixelMappingToolkit.cpp
--- a/Engine/Plugins/VirtualProduction/DMX/DMXPixelMapping/Source/DMXPixelMappingEditor/Private/Toolkits/DMXPixelMappingToolkit.cpp
+++ b/Engine/Plugins/VirtualProduction/DMX/DMXPixelMapping/Source/DMXPixelMappingEditor/Private/Toolkits/DMXPixelMappingToolkit.cpp
@@ -2,34 +2,31 @@
 
 #include "Toolkits/DMXPixelMappingToolkit.h"
 
-#include "DMXEditorUtils.h"
 #include "DMXPixelMapping.h"
 #include "DMXPixelMappingEditorCommands.h"
 #include "DMXPixelMappingEditorModule.h"
 #include "DMXPixelMappingEditorStyle.h"
 #include "DMXPixelMappingEditorUtils.h"
-#include "DMXPixelMappingLayoutSettings.h"
 #include "DMXPixelMappingToolbar.h"
+#include "Framework/Commands/GenericCommands.h"
+#include "Framework/MultiBox/MultiBoxExtender.h"
 #include "K2Node_PixelMappingBaseComponent.h"
 #include "Components/DMXPixelMappingFixtureGroupComponent.h"
 #include "Components/DMXPixelMappingRendererComponent.h"
 #include "Components/DMXPixelMappingRootComponent.h"
 #include "Components/DMXPixelMappingMatrixComponent.h"
 #include "Components/DMXPixelMappingScreenComponent.h"
+#include "Modules/ModuleManager.h"
+#include "ScopedTransaction.h"
+#include "Settings/DMXPixelMappingEditorSettings.h"
 #include "Templates/DMXPixelMappingComponentTemplate.h"
+#include "UObject/UObjectIterator.h"
 #include "Views/SDMXPixelMappingDesignerView.h"
 #include "Views/SDMXPixelMappingDetailsView.h"
 #include "Views/SDMXPixelMappingDMXLibraryView.h"
 #include "Views/SDMXPixelMappingHierarchyView.h"
 #include "Views/SDMXPixelMappingLayoutView.h"
 #include "Views/SDMXPixelMappingPreviewView.h"
-
-#include "ScopedTransaction.h"
-#include "Engine/TextureRenderTarget2D.h"
-#include "Framework/Commands/GenericCommands.h"
-#include "Framework/MultiBox/MultiBoxExtender.h"
-#include "Modules/ModuleManager.h"
-#include "UObject/UObjectIterator.h"
 #include "Widgets/Docking/SDockTab.h"
 
 
@@ -42,17 +39,15 @@ const FName FDMXPixelMappingToolkit::PreviewViewTabID(TEXT("DMXPixelMappingEdito
 const FName FDMXPixelMappingToolkit::DetailsViewTabID(TEXT("DMXPixelMappingEditor_DetailsViewTabID"));
 const FName FDMXPixelMappingToolkit::LayoutViewTabID(TEXT("DMXPixelMappingEditor_LayoutViewTabID"));
 
-
 FDMXPixelMappingToolkit::FDMXPixelMappingToolkit()
 	: DMXPixelMapping(nullptr)
-{}
+{
+	EditorSettingsDump = TArray<uint8, TFixedAllocator<sizeof(UDMXPixelMappingEditorSettings)>>(reinterpret_cast<const uint8*>(GetDefault<UDMXPixelMappingEditorSettings>()), (int32)sizeof(UDMXPixelMappingEditorSettings));
+}
 
 FDMXPixelMappingToolkit::~FDMXPixelMappingToolkit()
 {
-	UDMXPixelMappingBaseComponent::GetOnComponentAdded().RemoveAll(this);
-	UDMXPixelMappingBaseComponent::GetOnComponentRemoved().RemoveAll(this);
-
-	Toolbar.Reset();
+	SaveThumbnailImage();
 }
 
 void FDMXPixelMappingToolkit::InitPixelMappingEditor(const EToolkitMode::Type Mode, const TSharedPtr<class IToolkitHost>& InitToolkitHost, UDMXPixelMapping* InDMXPixelMapping)
@@ -161,6 +156,15 @@ void FDMXPixelMappingToolkit::Tick(float DeltaTime)
 		bIsPlayingDMX = false;
 		bRequestStopSendingDMX = false; 
 	}
+	
+	// Detect and broadcast editor setting changes
+	if (FMemory::Memcmp(EditorSettingsDump.GetData(), GetDefault<UDMXPixelMappingEditorSettings>(), sizeof(UDMXPixelMappingEditorSettings)) != 0)
+	{
+		UDMXPixelMappingEditorSettings* EditorSettings = GetMutableDefault<UDMXPixelMappingEditorSettings>();
+		EditorSettings->OnEditorSettingsChanged.Broadcast();
+
+		EditorSettingsDump = TArray<uint8, TFixedAllocator<sizeof(UDMXPixelMappingEditorSettings)>>(reinterpret_cast<const uint8*>(GetDefault<UDMXPixelMappingEditorSettings>()), (int32)sizeof(UDMXPixelMappingEditorSettings));
+	}
 }
 
 TStatId FDMXPixelMappingToolkit::GetStatId() const
@@ -195,15 +199,17 @@ void FDMXPixelMappingToolkit::SelectComponents(const TSet<FDMXPixelMappingCompon
 {
 	// Update selection
 	SelectedComponents.Empty();
-
-	SetActiveRenderComponent(nullptr);
 	ActiveOutputComponents.Empty();
 
 	SelectedComponents.Append(InSelectedComponents);
 
 	for (const FDMXPixelMappingComponentReference& ComponentReference : SelectedComponents)
 	{
-		if (UDMXPixelMappingRendererComponent* RendererComponent = Cast<UDMXPixelMappingRendererComponent>(ComponentReference.GetComponent()))
+		if (UDMXPixelMappingRootComponent* RootComponent = Cast<UDMXPixelMappingRootComponent>(ComponentReference.GetComponent()))
+		{
+			continue;
+		}
+		else if (UDMXPixelMappingRendererComponent* RendererComponent = Cast<UDMXPixelMappingRendererComponent>(ComponentReference.GetComponent()))
 		{
 			SetActiveRenderComponent(RendererComponent);
 		}
@@ -255,25 +261,20 @@ void FDMXPixelMappingToolkit::StopPlayingDMX()
 	bRequestStopSendingDMX = true;
 }
 
-void FDMXPixelMappingToolkit::ExecutebTogglePlayDMXAll()
+void FDMXPixelMappingToolkit::UpdateBlueprintNodes() const
 {
-	bTogglePlayDMXAll = !bTogglePlayDMXAll;
-}
-
-void FDMXPixelMappingToolkit::UpdateBlueprintNodes(UDMXPixelMapping* InDMXPixelMapping)
-{
-	if (InDMXPixelMapping != nullptr)
+	if (UDMXPixelMapping* PixelMapping = GetDMXPixelMapping())
 	{
 		for (TObjectIterator<UK2Node_PixelMappingBaseComponent> It(RF_Transient | RF_ClassDefaultObject, /** bIncludeDerivedClasses */ true, /** InternalExcludeFlags */ EInternalObjectFlags::Garbage); It; ++It)
 		{
-			It->OnPixelMappingChanged(InDMXPixelMapping);
+			It->OnPixelMappingChanged(PixelMapping);
 		}
 	}
 }
 
-void FDMXPixelMappingToolkit::OnSaveThumbnailImage()
+void FDMXPixelMappingToolkit::SaveThumbnailImage()
 {
-	if (ActiveRendererComponent.IsValid())
+	if (DMXPixelMapping && ActiveRendererComponent.IsValid())
 	{
 		DMXPixelMapping->ThumbnailImage = ActiveRendererComponent->GetRenderedInputTexture();
 	}
@@ -302,7 +303,7 @@ TArray<UDMXPixelMappingBaseComponent*> FDMXPixelMappingToolkit::CreateComponents
 			}
 		}
 
-		UpdateBlueprintNodes(DMXPixelMapping);
+		UpdateBlueprintNodes();
 	}
 
 	return NewComponents;
@@ -353,7 +354,7 @@ void FDMXPixelMappingToolkit::DeleteSelectedComponents()
 	// Select the Parent Components
 	SelectComponents(ParentComponentReferences);
 
-	UpdateBlueprintNodes(DMXPixelMapping);
+	UpdateBlueprintNodes();
 }	
 
 bool FDMXPixelMappingToolkit::CanSizeSelectedComponentToTexture() const
@@ -423,8 +424,8 @@ void FDMXPixelMappingToolkit::SizeSelectedComponentToTexture(bool bTransacted)
 		SizeComponentToTextureTransaction = MakeShared<FScopedTransaction>(LOCTEXT("SizeComponentToTextureTransaction", "Size Component to Texture"));
 	}
 
-	const UDMXPixelMappingLayoutSettings* LayoutSettings = GetDefault<UDMXPixelMappingLayoutSettings>();
-	if (LayoutSettings && LayoutSettings->bScaleChildrenWithParent)
+	const FDMXPixelMappingDesignerSettings& DesignerSettings = GetDefault<UDMXPixelMappingEditorSettings>()->DesignerSettings;
+	if (DesignerSettings.bScaleChildrenWithParent)
 	{
 		// Scale children to parent
 		const FVector2D RatioVector = TextureSize / Component->GetSize();
@@ -463,21 +464,17 @@ void FDMXPixelMappingToolkit::SizeSelectedComponentToTexture(bool bTransacted)
 
 void FDMXPixelMappingToolkit::OnComponentAdded(UDMXPixelMapping* PixelMapping, UDMXPixelMappingBaseComponent* Component)
 {
-	// For performance reasons don't update blueprint nodes if the methods in this class add (possibly many) components.
-	// Instead call UpdateBlueprintNodes at the end such functions.
 	if (!bAddingComponents)
 	{
-		UpdateBlueprintNodes(DMXPixelMapping);
+		UpdateBlueprintNodes();
 	}
 }
 
 void FDMXPixelMappingToolkit::OnComponentRemoved(UDMXPixelMapping* PixelMapping, UDMXPixelMappingBaseComponent* Component)
 {
-	// For performance reasons don't update blueprint nodes if the methods in this class remove (possibly many) components .
-	// Instead call UpdateBlueprintNodes at the end such functions.
 	if (!bRemovingComponents)
 	{
-		UpdateBlueprintNodes(DMXPixelMapping);
+		UpdateBlueprintNodes();
 	}
 }
 
@@ -492,8 +489,8 @@ void FDMXPixelMappingToolkit::InitializeInternal(const EToolkitMode::Type Mode,
 	DMXPixelMapping->CreateOrLoadObjects();
 
 	// Bind to component changes
-	UDMXPixelMappingBaseComponent::GetOnComponentAdded().AddRaw(this, &FDMXPixelMappingToolkit::OnComponentAdded);
-	UDMXPixelMappingBaseComponent::GetOnComponentRemoved().AddRaw(this, &FDMXPixelMappingToolkit::OnComponentRemoved);
+	UDMXPixelMappingBaseComponent::GetOnComponentAdded().AddSP(this, &FDMXPixelMappingToolkit::OnComponentAdded);
+	UDMXPixelMappingBaseComponent::GetOnComponentRemoved().AddSP(this, &FDMXPixelMappingToolkit::OnComponentRemoved);
 
 	// Create commands
 	DesignerCommandList = MakeShareable(new FUICommandList);
@@ -680,9 +677,29 @@ void FDMXPixelMappingToolkit::CreateInternalViews()
 	GetOrCreateLayoutView();
 }
 
-void FDMXPixelMappingToolkit::OnComponentRenamed(UDMXPixelMappingBaseComponent* InComponent)
+void FDMXPixelMappingToolkit::RenameComponent(const FName& CurrentObjectName, const FString& DesiredObjectName) const
 {
-	UpdateBlueprintNodes(GetDMXPixelMapping());
+	if (!DMXPixelMapping)
+	{
+		return;
+	}
+
+	UDMXPixelMappingBaseComponent* ComponentToRename = DMXPixelMapping->FindComponent(CurrentObjectName);
+	if (!ensureMsgf(ComponentToRename, TEXT("Cannot find component '%s' to rename."), *CurrentObjectName.ToString()))
+	{
+		return;
+	}
+
+	const FName DesiredDisplayName = MakeObjectNameFromDisplayLabel(DesiredObjectName, ComponentToRename->GetFName());
+	UDMXPixelMappingBaseComponent* ExistingComponent = DMXPixelMapping->FindComponent(DesiredDisplayName);
+
+	const FName UniqueName = ExistingComponent ?
+		MakeUniqueObjectName(ComponentToRename->GetOuter(), ComponentToRename->GetClass(), DesiredDisplayName) :
+		DesiredDisplayName;
+
+	ComponentToRename->Modify();
+	ComponentToRename->Rename(*UniqueName.ToString());
+	UpdateBlueprintNodes();
 }
 
 TSharedRef<SDMXPixelMappingDMXLibraryView> FDMXPixelMappingToolkit::GetOrCreateDMXLibraryView()
@@ -745,6 +762,23 @@ TSharedRef<SDMXPixelMappingLayoutView> FDMXPixelMappingToolkit::GetOrCreateLayou
 	return LayoutView.ToSharedRef();
 }
 
+#define UE_DMX_MAP_EDITOR_SETTING_TO_TOGGLE_COMMAND(EditorSettings, StructName, MemberName, Action) \
+	GetToolkitCommands()->MapAction( \
+		FDMXPixelMappingEditorCommands::Get().Action, \
+		FExecuteAction::CreateLambda([] \
+			{ \
+				UDMXPixelMappingEditorSettings* EditorSettings = GetMutableDefault<UDMXPixelMappingEditorSettings>(); \
+				EditorSettings->StructName.MemberName = !EditorSettings->StructName.MemberName; \
+				EditorSettings->SaveConfig(); \
+			}), \
+		FCanExecuteAction(), \
+		FIsActionChecked::CreateLambda([]() \
+			{  \
+				const UDMXPixelMappingEditorSettings* EditorSettings = GetDefault<UDMXPixelMappingEditorSettings>(); \
+				return EditorSettings->DesignerSettings.MemberName; \
+			}) \
+	); 
+
 void FDMXPixelMappingToolkit::SetupCommands()
 {
 	GetToolkitCommands()->MapAction(
@@ -752,11 +786,6 @@ void FDMXPixelMappingToolkit::SetupCommands()
 		FExecuteAction::CreateSP(this, &FDMXPixelMappingToolkit::AddRenderer)
 	);
 
-	GetToolkitCommands()->MapAction(
-		FDMXPixelMappingEditorCommands::Get().SaveThumbnailImage,
-		FExecuteAction::CreateSP(this, &FDMXPixelMappingToolkit::OnSaveThumbnailImage)
-	);
-
 	GetToolkitCommands()->MapAction(
 		FDMXPixelMappingEditorCommands::Get().PlayDMX,
 		FExecuteAction::CreateSP(this, &FDMXPixelMappingToolkit::PlayDMX),
@@ -785,17 +814,7 @@ void FDMXPixelMappingToolkit::SetupCommands()
 			})
 	);
 
-	GetToolkitCommands()->MapAction(
-		FDMXPixelMappingEditorCommands::Get().TogglePlayDMXAll,
-		FExecuteAction::CreateSP(this, &FDMXPixelMappingToolkit::ExecutebTogglePlayDMXAll),
-		FCanExecuteAction(),
-		FIsActionChecked::CreateLambda([this]()
-			{ 
-				return bTogglePlayDMXAll; 
-			})
-	);
-
-	// Layout related
+	// Designer related
 	constexpr bool bTransactSizeComponentToTexture = true;
 	GetToolkitCommands()->MapAction(
 		FDMXPixelMappingEditorCommands::Get().SizeComponentToTexture,
@@ -803,160 +822,16 @@ void FDMXPixelMappingToolkit::SetupCommands()
 		FCanExecuteAction::CreateSP(this, &FDMXPixelMappingToolkit::CanSizeSelectedComponentToTexture)
 	);
 
-	GetToolkitCommands()->MapAction(
-		FDMXPixelMappingEditorCommands::Get().ToggleScaleChildrenWithParent,
-		FExecuteAction::CreateLambda([]
-			{
-				if (UDMXPixelMappingLayoutSettings* LayoutSettings = GetMutableDefault<UDMXPixelMappingLayoutSettings>())
-				{
-					LayoutSettings->PreEditChange(nullptr);
-					LayoutSettings->bScaleChildrenWithParent = !LayoutSettings->bScaleChildrenWithParent;
-					LayoutSettings->PostEditChange();
-				}
-			}),
-		FCanExecuteAction(),
-		FIsActionChecked::CreateLambda([]()
-			{
-				if (const UDMXPixelMappingLayoutSettings* LayoutSettings = GetDefault<UDMXPixelMappingLayoutSettings>())
-				{
-					return LayoutSettings->bScaleChildrenWithParent;
-				}
-				return false;
-			})
-	);
-
-	GetToolkitCommands()->MapAction(
-		FDMXPixelMappingEditorCommands::Get().ToggleAlwaysSelectGroup,
-		FExecuteAction::CreateLambda([]
-			{
-				if (UDMXPixelMappingLayoutSettings* LayoutSettings = GetMutableDefault<UDMXPixelMappingLayoutSettings>())
-				{
-					LayoutSettings->PreEditChange(nullptr);
-					LayoutSettings->bAlwaysSelectGroup = !LayoutSettings->bAlwaysSelectGroup;
-					LayoutSettings->PostEditChange();
-				}
-			}),
-		FCanExecuteAction(),
-		FIsActionChecked::CreateLambda([]()
-			{
-				if (const UDMXPixelMappingLayoutSettings* LayoutSettings = GetDefault<UDMXPixelMappingLayoutSettings>())
-				{
-					return LayoutSettings->bAlwaysSelectGroup;
-				}
-				return false;
-			})
-	);
-
-	GetToolkitCommands()->MapAction(
-		FDMXPixelMappingEditorCommands::Get().ToggleApplyLayoutScriptWhenLoaded,
-		FExecuteAction::CreateLambda([]
-			{
-				if (UDMXPixelMappingLayoutSettings* LayoutSettings = GetMutableDefault<UDMXPixelMappingLayoutSettings>())
-				{
-					LayoutSettings->PreEditChange(nullptr);
-					LayoutSettings->bApplyLayoutScriptWhenLoaded = !LayoutSettings->bApplyLayoutScriptWhenLoaded;
-					LayoutSettings->PostEditChange();
-				}
-			}),
-		FCanExecuteAction(),
-		FIsActionChecked::CreateLambda([]()
-			{
-				if (const UDMXPixelMappingLayoutSettings* LayoutSettings = GetDefault<UDMXPixelMappingLayoutSettings>())
-				{
-					return LayoutSettings->bApplyLayoutScriptWhenLoaded;
-				}
-				return false;
-			})
-	);
-
-	GetToolkitCommands()->MapAction(
-		FDMXPixelMappingEditorCommands::Get().ToggleShowMatrixCells,
-		FExecuteAction::CreateLambda([]
-			{
-				if (UDMXPixelMappingLayoutSettings* LayoutSettings = GetMutableDefault<UDMXPixelMappingLayoutSettings>())
-				{
-					LayoutSettings->PreEditChange(nullptr);
-					LayoutSettings->bShowMatrixCells = !LayoutSettings->bShowMatrixCells;
-					LayoutSettings->PostEditChange();
-				}
-			}),
-		FCanExecuteAction(),
-				FIsActionChecked::CreateLambda([]()
-					{
-						if (const UDMXPixelMappingLayoutSettings* LayoutSettings = GetDefault<UDMXPixelMappingLayoutSettings>())
-						{
-							return LayoutSettings->bShowMatrixCells;
-						}
-						return false;
-					})
-				);
-
-	GetToolkitCommands()->MapAction(
-		FDMXPixelMappingEditorCommands::Get().ToggleShowComponentNames,
-		FExecuteAction::CreateLambda([]
-			{
-				if (UDMXPixelMappingLayoutSettings* LayoutSettings = GetMutableDefault<UDMXPixelMappingLayoutSettings>())
-				{
-					LayoutSettings->PreEditChange(nullptr);
-					LayoutSettings->bShowComponentNames = !LayoutSettings->bShowComponentNames;
-					LayoutSettings->PostEditChange();
-				}
-			}),
-		FCanExecuteAction(),
-		FIsActionChecked::CreateLambda([]()
-			{
-				if (const UDMXPixelMappingLayoutSettings* LayoutSettings = GetDefault<UDMXPixelMappingLayoutSettings>())
-				{
-					return LayoutSettings->bShowComponentNames;
-				}
-				return false;
-			})
-	);
-
-	GetToolkitCommands()->MapAction(
-		FDMXPixelMappingEditorCommands::Get().ToggleShowPatchInfo,
-		FExecuteAction::CreateLambda([]
-			{
-				if (UDMXPixelMappingLayoutSettings* LayoutSettings = GetMutableDefault<UDMXPixelMappingLayoutSettings>())
-				{
-					LayoutSettings->PreEditChange(nullptr);
-					LayoutSettings->bShowPatchInfo = !LayoutSettings->bShowPatchInfo;
-					LayoutSettings->PostEditChange();
-				}
-			}),
-		FCanExecuteAction(),
-		FIsActionChecked::CreateLambda([]()
-			{
-				if (const UDMXPixelMappingLayoutSettings* LayoutSettings = GetDefault<UDMXPixelMappingLayoutSettings>())
-				{
-					return LayoutSettings->bShowPatchInfo;
-				}
-				return false;
-			})
-	);
-
-	GetToolkitCommands()->MapAction(
-		FDMXPixelMappingEditorCommands::Get().ToggleShowCellIDs,
-		FExecuteAction::CreateLambda([]
-			{
-				if (UDMXPixelMappingLayoutSettings* LayoutSettings = GetMutableDefault<UDMXPixelMappingLayoutSettings>())
-				{
-					LayoutSettings->PreEditChange(nullptr);
-					LayoutSettings->bShowCellIDs = !LayoutSettings->bShowCellIDs;
-					LayoutSettings->PostEditChange();
-				}
-			}),
-		FCanExecuteAction(),
-		FIsActionChecked::CreateLambda([]()
-			{
-				if (const UDMXPixelMappingLayoutSettings* LayoutSettings = GetDefault<UDMXPixelMappingLayoutSettings>())
-				{
-					return LayoutSettings->bShowCellIDs;
-				}
-				return false;
-			})
-	);
+	UDMXPixelMappingEditorSettings* EditorSettings = GetMutableDefault<UDMXPixelMappingEditorSettings>(); 
+	UE_DMX_MAP_EDITOR_SETTING_TO_TOGGLE_COMMAND(EditorSettings, DesignerSettings, bScaleChildrenWithParent, ToggleScaleChildrenWithParent);
+	UE_DMX_MAP_EDITOR_SETTING_TO_TOGGLE_COMMAND(EditorSettings, DesignerSettings, bAlwaysSelectGroup, ToggleAlwaysSelectGroup);
+	UE_DMX_MAP_EDITOR_SETTING_TO_TOGGLE_COMMAND(EditorSettings, DesignerSettings, bApplyLayoutScriptWhenLoaded, ToggleApplyLayoutScriptWhenLoaded);
+	UE_DMX_MAP_EDITOR_SETTING_TO_TOGGLE_COMMAND(EditorSettings, DesignerSettings, bShowMatrixCells, ToggleShowMatrixCells);
+	UE_DMX_MAP_EDITOR_SETTING_TO_TOGGLE_COMMAND(EditorSettings, DesignerSettings, bShowComponentNames, ToggleShowComponentNames);
+	UE_DMX_MAP_EDITOR_SETTING_TO_TOGGLE_COMMAND(EditorSettings, DesignerSettings, bShowPatchInfo, ToggleShowPatchInfo);
+	UE_DMX_MAP_EDITOR_SETTING_TO_TOGGLE_COMMAND(EditorSettings, DesignerSettings, bShowCellIDs, ToggleShowCellIDs);
 }
+#undef UE_DMX_MAP_EDITOR_SETTING_TO_TOGGLE_COMMAND
 
 void FDMXPixelMappingToolkit::ExtendToolbar()
 {
diff --git a/Engine/Plugins/VirtualProduction/DMX/DMXPixelMapping/Source/DMXPixelMappingEditor/Private/Toolkits/DMXPixelMappingToolkit.h b/Engine/Plugins/VirtualProduction/DMX/DMXPixelMapping/Source/DMXPixelMappingEditor/Private/Toolkits/DMXPixelMappingToolkit.h
--- a/Engine/Plugins/VirtualProduction/DMX/DMXPixelMapping/Source/DMXPixelMappingEditor/Private/Toolkits/DMXPixelMappingToolkit.h
+++ b/Engine/Plugins/VirtualProduction/DMX/DMXPixelMapping/Source/DMXPixelMappingEditor/Private/Toolkits/DMXPixelMappingToolkit.h
@@ -3,14 +3,16 @@
 #pragma once
 
 #include "DMXPixelMappingComponentReference.h"
-
+#include "Settings/DMXPixelMappingEditorSettings.h"
 #include "TickableEditorObject.h"
 #include "Toolkits/AssetEditorToolkit.h"
+#include "Widgets/Views/SHeaderRow.h"
 
 class FDMXPixelMappingComponentTemplate;
-class FSpawnTabArgs;
-
 class FDMXPixelMappingToolbar;
+class FSpawnTabArgs;
+class FTabManager;
+class SDockableTab;
 class SDMXPixelMappingHierarchyView;
 class SDMXPixelMappingDesignerView;
 class SDMXPixelMappingDetailsView;
@@ -23,11 +25,6 @@ class UDMXPixelMappingMatrixComponent;
 class UDMXPixelMappingOutputComponent;
 class UDMXPixelMappingRendererComponent;
 
-struct FScopedSlowTask;
-class FTabManager;
-class SDockableTab;
-class SWidget;
-
 
 /**
  * Implements an Editor toolkit for Pixel Mapping.
@@ -36,8 +33,6 @@ class FDMXPixelMappingToolkit
 	: public FAssetEditorToolkit
 	, public FTickableEditorObject
 {
-	friend class FDMXPixelMappingToolbar;
-
 public:
 	DECLARE_MULTICAST_DELEGATE(FOnSelectedComponentsChangedDelegate)
 	FOnSelectedComponentsChangedDelegate& GetOnSelectedComponentsChangedDelegate() { return OnSelectedComponentsChangedDelegate; }
@@ -51,8 +46,6 @@ public:
 	 */
 	virtual ~FDMXPixelMappingToolkit();
 
-public:
-
 	/**
 	 * Edits the specified Texture object.
 	 *
@@ -63,7 +56,6 @@ public:
 	void InitPixelMappingEditor(const EToolkitMode::Type Mode, const TSharedPtr<class IToolkitHost>& InitToolkitHost, UDMXPixelMapping* InDMXPixelMapping);
 
 public:
-
 	//~ Begin FAssetEditorToolkit Interface
 	virtual void RegisterTabSpawners(const TSharedRef<FTabManager>& TabManager) override;
 	virtual void UnregisterTabSpawners(const TSharedRef<FTabManager>& TabManager) override;
@@ -88,19 +80,6 @@ public:
 
 	UDMXPixelMappingRendererComponent* GetActiveRendererComponent() const { return ActiveRendererComponent.Get(); }
 
-	TArray<UDMXPixelMappingOutputComponent*> GetActiveOutputComponents() const
-	{
-		TArray<UDMXPixelMappingOutputComponent*> ValidatedActiveOutputComponents;
-		for (TWeakObjectPtr<UDMXPixelMappingOutputComponent> ActiveOutputComponent : ActiveOutputComponents)
-		{
-			if (ActiveOutputComponent.IsValid())
-			{
-				ValidatedActiveOutputComponents.Add(ActiveOutputComponent.Get());
-			}
-		}
-		return ValidatedActiveOutputComponents;
-	}
-
 	const TSharedPtr<FUICommandList>& GetDesignerCommandList() const { return DesignerCommandList; }
 
 	const TSet<FDMXPixelMappingComponentReference>& GetSelectedComponents() const { return SelectedComponents; }
@@ -139,7 +118,13 @@ public:
 	/** Adds a renderer to the pixel mapping. Sets it as active renderer if a new renderer was created */
 	void AddRenderer();
 
-	void OnComponentRenamed(UDMXPixelMappingBaseComponent* InComponent);
+	/**
+	 * Rename Pixel Maping component.
+	 *
+	 * @param PreviousObjectName	The previous object name, used to find the object
+	 * @param DesiredObjectName		The desired object name
+	 */
+	void RenameComponent(const FName& PreviousObjectName, const FString& DesiredObjectName) const;
 
 	/** 
 	 * Creates components from the template. Returns the new components.
@@ -167,11 +152,11 @@ private:
 
 	void StopPlayingDMX();
 
-	void ExecutebTogglePlayDMXAll();
+	/** Updates blueprint nodes */
+	void UpdateBlueprintNodes() const;
 
-	void UpdateBlueprintNodes(UDMXPixelMapping* InDMXPixelMapping);
-
-	void OnSaveThumbnailImage();
+	/** Saves a thumbnail image for the pixel mapping asset */
+	void SaveThumbnailImage();
 
 	void InitializeInternal(const EToolkitMode::Type Mode, const TSharedPtr<class IToolkitHost>& InitToolkitHost, const FGuid& MessageLogGuid);
 
@@ -204,7 +189,7 @@ private:
 	/** List of open tool panels; used to ensure only one exists at any one time */
 	TMap<FName, TWeakPtr<SDockableTab>> SpawnedToolPanels;
 
-	/** The Input Source View instance */
+	/** The DMX Library View instance */
 	TSharedPtr<SDMXPixelMappingDMXLibraryView> DMXLibraryView;
 
 	/** The Hierarchy View instance */
@@ -268,4 +253,7 @@ public:
 
 	/** Name of the Details View Tab */
 	static const FName LayoutViewTabID;
+
+	/** Dumped editor settings for fast comparison */
+	TArray<uint8, TFixedAllocator<sizeof(UDMXPixelMappingEditorSettings)>> EditorSettingsDump;
 };
diff --git a/Engine/Plugins/VirtualProduction/DMX/DMXPixelMapping/Source/DMXPixelMappingEditor/Private/ViewModels/DMXPixelMappingDMXLibraryViewModel.cpp b/Engine/Plugins/VirtualProduction/DMX/DMXPixelMapping/Source/DMXPixelMappingEditor/Private/ViewModels/DMXPixelMappingDMXLibraryViewModel.cpp
--- a/Engine/Plugins/VirtualProduction/DMX/DMXPixelMapping/Source/DMXPixelMappingEditor/Private/ViewModels/DMXPixelMappingDMXLibraryViewModel.cpp
+++ b/Engine/Plugins/VirtualProduction/DMX/DMXPixelMapping/Source/DMXPixelMappingEditor/Private/ViewModels/DMXPixelMappingDMXLibraryViewModel.cpp
@@ -92,7 +92,7 @@ void UDMXPixelMappingDMXLibraryViewModel::UpdateFixtureGroupFromSelection(TWeakP
 		{
 			if (Component->GetClass() == UDMXPixelMappingFixtureGroupComponent::StaticClass())
 			{
-				FixtureGroupComponents.Add(CastChecked<UDMXPixelMappingFixtureGroupComponent>(Component));
+				FixtureGroupComponents.AddUnique(CastChecked<UDMXPixelMappingFixtureGroupComponent>(Component));
 			}
 		}
 	}
diff --git a/Engine/Plugins/VirtualProduction/DMX/DMXPixelMapping/Source/DMXPixelMappingEditor/Private/ViewModels/DMXPixelMappingHierarchyItem.cpp b/Engine/Plugins/VirtualProduction/DMX/DMXPixelMapping/Source/DMXPixelMappingEditor/Private/ViewModels/DMXPixelMappingHierarchyItem.cpp
new file mode 100644
--- /dev/null
+++ b/Engine/Plugins/VirtualProduction/DMX/DMXPixelMapping/Source/DMXPixelMappingEditor/Private/ViewModels/DMXPixelMappingHierarchyItem.cpp
@@ -0,0 +1,170 @@
+// Copyright Epic Games, Inc. All Rights Reserved.
+
+#include "ViewModels/DMXPixelMappingHierarchyItem.h"
+
+#include "Algo/Count.h"
+#include "Algo/Reverse.h"
+#include "Algo/Sort.h"
+#include "Components/DMXPixelMappingOutputComponent.h"
+#include "Components/DMXPixelMappingOutputDMXComponent.h"
+#include "Components/DMXPixelMappingRootComponent.h"
+#include "DMXPixelMapping.h"
+#include "Library/DMXEntityFixturePatch.h"
+#include "Toolkits/DMXPixelMappingToolkit.h"
+
+
+TSharedRef<FDMXPixelMappingHierarchyItem> FDMXPixelMappingHierarchyItem::CreateNew(TWeakPtr<FDMXPixelMappingToolkit> InToolkit)
+{
+	const TSharedRef<FDMXPixelMappingHierarchyItem> NewObject = MakeShareable(new FDMXPixelMappingHierarchyItem(InToolkit));
+	NewObject->Initialize();
+
+	return NewObject;
+}
+
+TArray<TSharedPtr<FDMXPixelMappingHierarchyItem>> FDMXPixelMappingHierarchyItem::GetItemAndChildrenRecursive()
+{
+	TArray<TSharedPtr<FDMXPixelMappingHierarchyItem>> Result;
+	Result.Add(AsShared());
+
+	for (const TSharedPtr<FDMXPixelMappingHierarchyItem>& Item : Children)
+	{
+		Result.Append(Item->GetItemAndChildrenRecursive());
+	}
+
+	return Result;
+}
+
+bool FDMXPixelMappingHierarchyItem::HasEditorColor() const
+{
+	return WeakComponent.Get() ? WeakComponent->IsA(UDMXPixelMappingOutputDMXComponent::StaticClass()) : false;
+}
+
+FLinearColor FDMXPixelMappingHierarchyItem::GetEditorColor() const
+{
+	if (UDMXPixelMappingOutputDMXComponent* OutputComponent = Cast<UDMXPixelMappingOutputDMXComponent>(WeakComponent.Get()))
+	{
+		return OutputComponent->GetEditorColor();
+	}
+	return FLinearColor::Red;
+}
+
+FText FDMXPixelMappingHierarchyItem::GetComponentNameText() const
+{
+	return WeakComponent.IsValid() ?
+		FText::FromString(WeakComponent->GetUserFriendlyName()) :
+		FText::GetEmpty();
+}
+
+FText FDMXPixelMappingHierarchyItem::GetFixtureIDText() const
+{
+	if (OptionalFixtureID.IsSet())
+	{
+		return FText::FromString(FString::FromInt(OptionalFixtureID.GetValue()));
+	}
+
+	return FText::GetEmpty();
+}
+
+FText FDMXPixelMappingHierarchyItem::GetPatchText() const
+{
+	if (UDMXPixelMappingOutputDMXComponent* OutputDMXComponent = Cast<UDMXPixelMappingOutputDMXComponent>(WeakComponent.Get()))
+	{
+		if (UDMXEntityFixturePatch* FixturePatch = OutputDMXComponent->FixturePatchRef.GetFixturePatch())
+		{
+			const FString PatchString = FString::FromInt(FixturePatch->GetUniverseID()) + TEXT(".") + FString::FromInt(FixturePatch->GetStartingChannel());
+			return FText::FromString(PatchString);
+		}
+	}
+	return FText::GetEmpty();
+}
+
+int64 FDMXPixelMappingHierarchyItem::GetAbsoluteChannel() const
+{
+	if (UDMXPixelMappingOutputDMXComponent* OutputDMXComponent = Cast<UDMXPixelMappingOutputDMXComponent>(WeakComponent.Get()))
+	{
+		if (UDMXEntityFixturePatch* FixturePatch = OutputDMXComponent->FixturePatchRef.GetFixturePatch())
+		{
+			return FixturePatch->GetUniverseID() * DMX_UNIVERSE_SIZE + FixturePatch->GetStartingChannel();
+		}
+	}
+	return -1;
+}
+
+void FDMXPixelMappingHierarchyItem::ReverseChildren()
+{
+	Algo::Reverse(Children);
+}
+
+FDMXPixelMappingHierarchyItem::FDMXPixelMappingHierarchyItem(TWeakPtr<FDMXPixelMappingToolkit> InToolkit)
+	: WeakToolkit(InToolkit)
+{
+	const TSharedPtr<FDMXPixelMappingToolkit> Toolkit = WeakToolkit.Pin();
+	if (!Toolkit.IsValid())
+	{
+		return;
+	}
+
+	UDMXPixelMapping* DMXPixelMapping = Toolkit->GetDMXPixelMapping();
+	if (DMXPixelMapping && DMXPixelMapping->RootComponent)
+	{
+		WeakComponent = DMXPixelMapping->RootComponent;
+	}
+}
+
+FDMXPixelMappingHierarchyItem::FDMXPixelMappingHierarchyItem(TWeakPtr<FDMXPixelMappingToolkit> InToolkit, UDMXPixelMappingBaseComponent* InComponent)
+	: WeakComponent(InComponent)
+	, WeakToolkit(InToolkit)
+{}
+
+void FDMXPixelMappingHierarchyItem::Initialize()
+{
+	BuildChildren();
+	UpdateFixtureID();
+
+	UDMXEntityFixturePatch::GetOnFixturePatchChanged().AddSP(this, &FDMXPixelMappingHierarchyItem::OnFixturePatchChanged);
+}
+
+void FDMXPixelMappingHierarchyItem::BuildChildren(UDMXPixelMappingBaseComponent* ParentPixelMappingComponent)
+{
+	const TSharedPtr<FDMXPixelMappingToolkit> Toolkit = WeakToolkit.Pin();
+	UDMXPixelMappingBaseComponent* Component = WeakComponent.Get();
+	if (!Toolkit.IsValid() || !Component)
+	{
+		return;
+	}
+
+	for (UDMXPixelMappingBaseComponent* ChildComponent : Component->GetChildren())
+	{
+		if (ChildComponent)
+		{
+			const TSharedRef<FDMXPixelMappingHierarchyItem> ChildItem = MakeShareable(new FDMXPixelMappingHierarchyItem(Toolkit, ChildComponent));
+			ChildItem->Initialize();
+
+			Children.Add(ChildItem);
+		}
+	}
+}
+
+void FDMXPixelMappingHierarchyItem::OnFixturePatchChanged(const UDMXEntityFixturePatch* FixturePatch)
+{
+	UpdateFixtureID();
+}
+
+void FDMXPixelMappingHierarchyItem::UpdateFixtureID()
+{
+	if (UDMXPixelMappingOutputDMXComponent* OutputDMXComponent = Cast<UDMXPixelMappingOutputDMXComponent>(WeakComponent.Get()))
+	{
+		if (UDMXEntityFixturePatch* FixturePatch = OutputDMXComponent->FixturePatchRef.GetFixturePatch())
+		{
+			int32 FixtureID;
+			if (FixturePatch->FindFixtureID(FixtureID))
+			{
+				OptionalFixtureID = FixtureID;
+			}
+			else
+			{
+				OptionalFixtureID.Reset();
+			}
+		}
+	}
+}
\ No newline at end of file
diff --git a/Engine/Plugins/VirtualProduction/DMX/DMXPixelMapping/Source/DMXPixelMappingEditor/Private/ViewModels/DMXPixelMappingHierarchyItem.h b/Engine/Plugins/VirtualProduction/DMX/DMXPixelMapping/Source/DMXPixelMappingEditor/Private/ViewModels/DMXPixelMappingHierarchyItem.h
new file mode 100644
--- /dev/null
+++ b/Engine/Plugins/VirtualProduction/DMX/DMXPixelMapping/Source/DMXPixelMappingEditor/Private/ViewModels/DMXPixelMappingHierarchyItem.h
@@ -0,0 +1,96 @@
+// Copyright Epic Games, Inc. All Rights Reserved.
+
+#pragma once
+
+#include "Algo/Sort.h"
+#include "DMXPixelMappingEditorCommon.h"
+#include "DMXPixelMappingComponentReference.h"
+
+struct FSlateColor;
+class UDMXEntityFixturePatch;
+class UDMXPixelMappingBaseComponent;
+
+
+/** 
+ * Model for the items in the hierarchy tree view.
+ * Children populate by construction and update on changes.
+ * Use OnModelChanged to listen to changes.
+ */
+class FDMXPixelMappingHierarchyItem final
+	: public TSharedFromThis<FDMXPixelMappingHierarchyItem>
+{
+public:
+	/** Creates a new instance */
+	static TSharedRef<FDMXPixelMappingHierarchyItem> CreateNew(TWeakPtr<FDMXPixelMappingToolkit> InToolkit);
+
+	/** Returns this item and all its children */
+	TArray<TSharedPtr<FDMXPixelMappingHierarchyItem>> GetItemAndChildrenRecursive();
+
+	/** Returns true if this row has an editor color */
+	bool HasEditorColor() const;
+
+	/** Returns the editor color. Optional return value is only set if the component can display an editor color */
+	FLinearColor GetEditorColor() const;
+
+	/** Returns the component name */
+	FText GetComponentNameText() const;
+
+	/** Returns the fixture ID text */
+	FText GetFixtureIDText() const;
+
+	/** Returns the patch text */
+	FText GetPatchText() const;
+
+	/** Gets the absolute channel of the patch of the component. Returns -1 if this item doesn't have a patch */
+	int64 GetAbsoluteChannel() const;
+
+	/** Sorts the children of this item by Projection. */
+	template <typename ProjectionType>
+	void StableSortChildren(ProjectionType Projection)
+	{
+		Algo::StableSortBy(Children, Projection);
+	}
+
+	/** Reverses the children of this item */
+	void ReverseChildren();
+
+	/** Returns the component of this model, or nullptr if the component is no longer valid */
+	UDMXPixelMappingBaseComponent* GetComponent() const { return WeakComponent.Get(); }
+
+	/** Returns the child models of this model (the child components of the component of this model) */
+	TArray<TSharedPtr<FDMXPixelMappingHierarchyItem>> GetChildren() const { return Children; }
+
+	/** Delegate any object can raise, causes the table row that draws this items to enter rename mode */
+	FSimpleMulticastDelegate RequestEnterRenameMode;
+
+private:
+	/** Constructs this item from the root component of the toolkit */
+	FDMXPixelMappingHierarchyItem(TWeakPtr<FDMXPixelMappingToolkit> InToolkit);
+
+	/** Constructst this item from the specified component in the component hierarchy */
+	FDMXPixelMappingHierarchyItem(TWeakPtr<FDMXPixelMappingToolkit> InToolkit, UDMXPixelMappingBaseComponent* InComponentReference);
+
+	/** Common initializer for constructors */
+	void Initialize();
+
+	/** Builds the Children member. If ParentPixelMappingComponent uses the root component as parent */
+	void BuildChildren(UDMXPixelMappingBaseComponent* ParentPixelMappingComponent = nullptr);
+
+	/** Called when a fixture patch changed */
+	void OnFixturePatchChanged(const UDMXEntityFixturePatch* FixturePatch);
+
+	/** Updates the optional fixture ID */
+	void UpdateFixtureID();
+
+	/** Optional fixture ID */
+	TOptional<int32> OptionalFixtureID;
+
+	/** The child items in the component hierarchy */
+	TArray<TSharedPtr<FDMXPixelMappingHierarchyItem>> Children;
+
+	/** Reference to the component this model handles */
+	TWeakObjectPtr<UDMXPixelMappingBaseComponent> WeakComponent;
+
+	/** Toolkit this model resides in */
+	FDMXPixelMappingToolkitWeakPtr WeakToolkit;
+};
diff --git a/Engine/Plugins/VirtualProduction/DMX/DMXPixelMapping/Source/DMXPixelMappingEditor/Private/ViewModels/DMXPixelMappingHierarchyViewModel.cpp b/Engine/Plugins/VirtualProduction/DMX/DMXPixelMapping/Source/DMXPixelMappingEditor/Private/ViewModels/DMXPixelMappingHierarchyViewModel.cpp
deleted file mode 100644
diff --git a/Engine/Plugins/VirtualProduction/DMX/DMXPixelMapping/Source/DMXPixelMappingEditor/Private/ViewModels/DMXPixelMappingHierarchyViewModel.h b/Engine/Plugins/VirtualProduction/DMX/DMXPixelMapping/Source/DMXPixelMappingEditor/Private/ViewModels/DMXPixelMappingHierarchyViewModel.h
deleted file mode 100644
diff --git a/Engine/Plugins/VirtualProduction/DMX/DMXPixelMapping/Source/DMXPixelMappingEditor/Private/ViewModels/DMXPixelMappingLayoutViewModel.cpp b/Engine/Plugins/VirtualProduction/DMX/DMXPixelMapping/Source/DMXPixelMappingEditor/Private/ViewModels/DMXPixelMappingLayoutViewModel.cpp
--- a/Engine/Plugins/VirtualProduction/DMX/DMXPixelMapping/Source/DMXPixelMappingEditor/Private/ViewModels/DMXPixelMappingLayoutViewModel.cpp
+++ b/Engine/Plugins/VirtualProduction/DMX/DMXPixelMapping/Source/DMXPixelMappingEditor/Private/ViewModels/DMXPixelMappingLayoutViewModel.cpp
@@ -4,16 +4,15 @@
 
 #include "Components/DMXPixelMappingMatrixCellComponent.h"
 #include "DMXPixelMappingEditorCommon.h"
-#include "DMXPixelMappingLayoutSettings.h"
 #include "Components/DMXPixelMappingFixtureGroupComponent.h"
 #include "Components/DMXPixelMappingFixtureGroupItemComponent.h"
 #include "Components/DMXPixelMappingMatrixComponent.h"
 #include "Components/DMXPixelMappingRendererComponent.h"
 #include "Components/DMXPixelMappingScreenComponent.h"
 #include "LayoutScripts/DMXPixelMappingLayoutScript.h"
-#include "Toolkits/DMXPixelMappingToolkit.h"
-
 #include "ScopedTransaction.h"
+#include "Settings/DMXPixelMappingEditorSettings.h"
+#include "Toolkits/DMXPixelMappingToolkit.h"
 
 
 #define LOCTEXT_NAMESPACE "DMXPixelMappingLayoutViewModel"
@@ -62,8 +61,8 @@ void UDMXPixelMappingLayoutViewModel::SetToolkit(const TSharedRef<FDMXPixelMappi
 	RefreshComponents();
 	RefreshLayoutScriptClass();
 
-	const UDMXPixelMappingLayoutSettings* LayoutSettings = GetDefault<UDMXPixelMappingLayoutSettings>();
-	if (LayoutSettings && LayoutSettings->bApplyLayoutScriptWhenLoaded)
+	const FDMXPixelMappingDesignerSettings& DesignerSettings = GetDefault<UDMXPixelMappingEditorSettings>()->DesignerSettings;
+	if (DesignerSettings.bApplyLayoutScriptWhenLoaded)
 	{
 		FScopedTransaction ApplyLayoutScriptTransactionDirect(LOCTEXT("ApplyLayoutScriptTransactionDirect", "Apply Layout Script"));
 		ApplyLayoutScripts();
@@ -315,8 +314,8 @@ void UDMXPixelMappingLayoutViewModel::OnLayoutScriptClassChanged()
 		OnModelChanged.Broadcast();
 	}
 
-	const UDMXPixelMappingLayoutSettings* LayoutSettings = GetDefault<UDMXPixelMappingLayoutSettings>();
-	if (LayoutScriptClass && LayoutSettings && LayoutSettings->bApplyLayoutScriptWhenLoaded)
+	const FDMXPixelMappingDesignerSettings& DesignerSettings = GetDefault<UDMXPixelMappingEditorSettings>()->DesignerSettings;
+	if (DesignerSettings.bApplyLayoutScriptWhenLoaded && LayoutScriptClass)
 	{
 		ApplyLayoutScripts();
 	}
diff --git a/Engine/Plugins/VirtualProduction/DMX/DMXPixelMapping/Source/DMXPixelMappingEditor/Private/ViewModels/DMXPixelMappingOutputComponentModel.cpp b/Engine/Plugins/VirtualProduction/DMX/DMXPixelMapping/Source/DMXPixelMappingEditor/Private/ViewModels/DMXPixelMappingOutputComponentModel.cpp
--- a/Engine/Plugins/VirtualProduction/DMX/DMXPixelMapping/Source/DMXPixelMappingEditor/Private/ViewModels/DMXPixelMappingOutputComponentModel.cpp
+++ b/Engine/Plugins/VirtualProduction/DMX/DMXPixelMapping/Source/DMXPixelMappingEditor/Private/ViewModels/DMXPixelMappingOutputComponentModel.cpp
@@ -3,7 +3,6 @@
 #include "ViewModels/DMXPixelMappingOutputComponentModel.h"
 
 #include "Components/DMXPixelMappingFixtureGroupComponent.h"
-#include "DMXPixelMappingLayoutSettings.h"
 #include "DMXPixelMappingTypes.h"
 #include "Components/DMXPixelMappingFixtureGroupItemComponent.h"
 #include "Components/DMXPixelMappingMatrixCellComponent.h"
@@ -13,6 +12,7 @@
 #include "Library/DMXLibrary.h"
 #include "MVR/DMXMVRGeneralSceneDescription.h"
 #include "MVR/Types/DMXMVRFixtureNode.h"
+#include "Settings/DMXPixelMappingEditorSettings.h"
 #include "Toolkits/DMXPixelMappingToolkit.h"
 
 
@@ -75,8 +75,8 @@ bool FDMXPixelMappingOutputComponentModel::ShouldDraw() const
 {
 	if (UDMXPixelMappingMatrixCellComponent* CellComponent = Cast<UDMXPixelMappingMatrixCellComponent>(WeakOutputComponent.Get()))
 	{
-		const UDMXPixelMappingLayoutSettings* LayoutSettings = GetDefault<UDMXPixelMappingLayoutSettings>();
-		return LayoutSettings->bShowMatrixCells;
+		const FDMXPixelMappingDesignerSettings& DesignerSettings = GetDefault<UDMXPixelMappingEditorSettings>()->DesignerSettings;
+		return DesignerSettings.bShowMatrixCells;
 	}
 
 	return true;
@@ -84,10 +84,9 @@ bool FDMXPixelMappingOutputComponentModel::ShouldDraw() const
 
 bool FDMXPixelMappingOutputComponentModel::ShouldDrawName() const
 {
-	const UDMXPixelMappingLayoutSettings* LayoutSettings = GetDefault<UDMXPixelMappingLayoutSettings>();
-	check(LayoutSettings);
+	const FDMXPixelMappingDesignerSettings& DesignerSettings = GetDefault<UDMXPixelMappingEditorSettings>()->DesignerSettings;
 
-	return LayoutSettings->bShowComponentNames &&
+	return DesignerSettings.bShowComponentNames &&
 		WeakOutputComponent.IsValid() && 
 		WeakOutputComponent->GetClass() != UDMXPixelMappingMatrixCellComponent::StaticClass();
 }
@@ -105,22 +104,20 @@ bool FDMXPixelMappingOutputComponentModel::ShouldDrawNameAbove() const
 
 bool FDMXPixelMappingOutputComponentModel::ShouldDrawCellID() const
 {
-	const UDMXPixelMappingLayoutSettings* LayoutSettings = GetDefault<UDMXPixelMappingLayoutSettings>();
-	check(LayoutSettings);
+	const FDMXPixelMappingDesignerSettings& DesignerSettings = GetDefault<UDMXPixelMappingEditorSettings>()->DesignerSettings;
 
 	return
-		LayoutSettings->bShowCellIDs &&
+		DesignerSettings.bShowCellIDs &&
 		WeakOutputComponent.IsValid() &&
 		WeakOutputComponent->GetClass() == UDMXPixelMappingMatrixCellComponent::StaticClass();
 }
 
 bool FDMXPixelMappingOutputComponentModel::ShouldDrawPatchInfo() const
 {
-	const UDMXPixelMappingLayoutSettings* LayoutSettings = GetDefault<UDMXPixelMappingLayoutSettings>();
-	check(LayoutSettings);
+	const FDMXPixelMappingDesignerSettings& DesignerSettings = GetDefault<UDMXPixelMappingEditorSettings>()->DesignerSettings;
 
 	return
-		LayoutSettings->bShowPatchInfo &&
+		DesignerSettings.bShowPatchInfo &&
 		GetFixturePatch() != nullptr;;
 }
 
diff --git a/Engine/Plugins/VirtualProduction/DMX/DMXPixelMapping/Source/DMXPixelMappingEditor/Private/Views/SDMXPixelMappingDMXLibraryView.cpp b/Engine/Plugins/VirtualProduction/DMX/DMXPixelMapping/Source/DMXPixelMappingEditor/Private/Views/SDMXPixelMappingDMXLibraryView.cpp
--- a/Engine/Plugins/VirtualProduction/DMX/DMXPixelMapping/Source/DMXPixelMappingEditor/Private/Views/SDMXPixelMappingDMXLibraryView.cpp
+++ b/Engine/Plugins/VirtualProduction/DMX/DMXPixelMapping/Source/DMXPixelMappingEditor/Private/Views/SDMXPixelMappingDMXLibraryView.cpp
@@ -83,7 +83,7 @@ void SDMXPixelMappingDMXLibraryView::Construct(const FArguments& InArgs, const T
 					{
 						return bRenderComponentContainedInSelection ?
 							LOCTEXT("AddFixtureGroupTooltip", "Adds a DMX Library to the Pixel Mapping") :
-							LOCTEXT("CannotAddFixtureGroupTooltip", "Please select an Input Source");
+							LOCTEXT("CannotAddFixtureGroupTooltip", "Please select a source texture");
 					})
 				.ContentPadding(FMargin(5.0f, 1.0f))
 				.OnClicked(this, &SDMXPixelMappingDMXLibraryView::OnAddDMXLibraryButtonClicked)
diff --git a/Engine/Plugins/VirtualProduction/DMX/DMXPixelMapping/Source/DMXPixelMappingEditor/Private/Views/SDMXPixelMappingDesignerView.cpp b/Engine/Plugins/VirtualProduction/DMX/DMXPixelMapping/Source/DMXPixelMappingEditor/Private/Views/SDMXPixelMappingDesignerView.cpp
--- a/Engine/Plugins/VirtualProduction/DMX/DMXPixelMapping/Source/DMXPixelMappingEditor/Private/Views/SDMXPixelMappingDesignerView.cpp
+++ b/Engine/Plugins/VirtualProduction/DMX/DMXPixelMapping/Source/DMXPixelMappingEditor/Private/Views/SDMXPixelMappingDesignerView.cpp
@@ -5,7 +5,6 @@
 #include "Algo/Find.h"
 #include "DMXPixelMapping.h"
 #include "DMXPixelMappingEditorStyle.h"
-#include "DMXPixelMappingLayoutSettings.h"
 #include "DragDrop/DMXPixelMappingGroupChildDragDropHelper.h"
 #include "Framework/Application/SlateApplication.h"
 #include "Components/DMXPixelMappingRendererComponent.h"
@@ -14,21 +13,21 @@
 #include "Components/DMXPixelMappingMatrixCellComponent.h"
 #include "Components/DMXPixelMappingScreenComponent.h"
 #include "DragDrop/DMXPixelMappingDragDropOp.h"
+#include "Input/HittestGrid.h"
+#include "ScopedTransaction.h"
+#include "Settings/DMXPixelMappingEditorSettings.h"
+#include "Toolkits/DMXPixelMappingToolkit.h"
+#include "Widgets/Input/SButton.h"
+#include "Widgets/Layout/SSpacer.h"
+#include "Widgets/Layout/SBox.h"
+#include "Widgets/Layout/SGridPanel.h"
+#include "Widgets/SCanvas.h"
 #include "Widgets/SDMXPixelMappingDesignerCanvas.h"
 #include "Widgets/SDMXPixelMappingOutputComponent.h"
 #include "Widgets/SDMXPixelMappingRuler.h"
 #include "Widgets/SDMXPixelMappingSourceTextureViewport.h"
 #include "Widgets/SDMXPixelMappingTransformHandle.h"
 #include "Widgets/SDMXPixelMappingZoomPan.h"
-#include "Toolkits/DMXPixelMappingToolkit.h"
-
-#include "ScopedTransaction.h"
-#include "Widgets/SCanvas.h"
-#include "Widgets/Layout/SBox.h"
-#include "Widgets/Layout/SGridPanel.h"
-#include "Input/HittestGrid.h"
-#include "Widgets/Layout/SSpacer.h"
-#include "Widgets/Input/SButton.h"
 
 
 #define LOCTEXT_NAMESPACE "SDMXPixelMappingDesignerView"
@@ -886,16 +885,12 @@ void SDMXPixelMappingDesignerView::ResolvePendingSelectedComponents(bool bClearP
 		}
 
 		// Select group if desired
-		const UDMXPixelMappingLayoutSettings* LayoutSettings = GetDefault<UDMXPixelMappingLayoutSettings>();
-		if (LayoutSettings && LayoutSettings->bAlwaysSelectGroup)
+		const FDMXPixelMappingDesignerSettings& DesignerSettings = GetDefault<UDMXPixelMappingEditorSettings>()->DesignerSettings;
+		if (DesignerSettings.bAlwaysSelectGroup)
 		{
-			if (UDMXPixelMappingFixtureGroupItemComponent* GroupItemComponent = Cast<UDMXPixelMappingFixtureGroupItemComponent>(PendingSelectedComponent))
-			{
-				PendingSelectedComponent = GroupItemComponent->GetParent();
-			}
-			else if (UDMXPixelMappingMatrixComponent* MatrixComponent = Cast<UDMXPixelMappingMatrixComponent>(PendingSelectedComponent))
+			if (UDMXPixelMappingOutputDMXComponent* OutputDMXComponent = Cast<UDMXPixelMappingFixtureGroupItemComponent>(PendingSelectedComponent))
 			{
-				PendingSelectedComponent = MatrixComponent->GetParent();
+				PendingSelectedComponent = OutputDMXComponent->GetParent();
 			}
 		}
 
@@ -905,7 +900,7 @@ void SDMXPixelMappingDesignerView::ResolvePendingSelectedComponents(bool bClearP
 
 		if (!bClearPreviousSelection)
 		{
-			SelectedComponents.Append(WeakToolkit.Pin()->GetSelectedComponents());
+			SelectedComponents.Append(Toolkit->GetSelectedComponents());
 		}
 		Toolkit->SelectComponents(SelectedComponents);
 
diff --git a/Engine/Plugins/VirtualProduction/DMX/DMXPixelMapping/Source/DMXPixelMappingEditor/Private/Views/SDMXPixelMappingHierarchyView.cpp b/Engine/Plugins/VirtualProduction/DMX/DMXPixelMapping/Source/DMXPixelMappingEditor/Private/Views/SDMXPixelMappingHierarchyView.cpp
--- a/Engine/Plugins/VirtualProduction/DMX/DMXPixelMapping/Source/DMXPixelMappingEditor/Private/Views/SDMXPixelMappingHierarchyView.cpp
+++ b/Engine/Plugins/VirtualProduction/DMX/DMXPixelMapping/Source/DMXPixelMappingEditor/Private/Views/SDMXPixelMappingHierarchyView.cpp
@@ -2,7 +2,9 @@
 
 #include "Views/SDMXPixelMappingHierarchyView.h"
 
+#include "Components/DMXPixelMappingFixtureGroupComponent.h"
 #include "Components/DMXPixelMappingMatrixComponent.h"
+#include "Components/DMXPixelMappingMatrixCellComponent.h"
 #include "Components/DMXPixelMappingRendererComponent.h"
 #include "Components/DMXPixelMappingRootComponent.h"
 #include "DMXPixelMapping.h"
@@ -17,272 +19,70 @@
 #include "HAL/PlatformApplicationMisc.h"
 #include "Misc/TextFilter.h"
 #include "ScopedTransaction.h"
+#include "Settings/DMXPixelMappingEditorSettings.h"
 #include "Toolkits/DMXPixelMappingToolkit.h"
 #include "UnrealExporter.h"
+#include "ViewModels/DMXPixelMappingHierarchyItem.h"
+#include "Views/SDMXPixelMappingDesignerView.h"
+#include "Views/SDMXPixelMappingPreviewView.h"
 #include "Widgets/Input/SButton.h"
+#include "Widgets/Input/SComboButton.h"
+#include "Widgets/Input/SSearchBox.h"
 #include "Widgets/Layout/SScrollBorder.h"
 #include "Widgets/SDMXPixelMappingHierarchyItem.h"
-#include "ViewModels/DMXPixelMappingHierarchyViewModel.h"
-#include "Views/SDMXPixelMappingDesignerView.h"
-#include "Views/SDMXPixelMappingPreviewView.h"
 #include "Widgets/Views/STreeView.h"
 
-#define LOCTEXT_NAMESPACE "SDMXPixelMappingHierarchyView"
 
+#define LOCTEXT_NAMESPACE "SDMXPixelMappingHierarchyView"
 
-class FDMXPixelMappingBaseComponentTextFactory 
-	: public FCustomizableTextObjectFactory
+namespace UE::DMX::PixelMappingEditor::SDMXPixelMappingHierarchyView::Private
 {
-public:
-
-	FDMXPixelMappingBaseComponentTextFactory() 
-		: FCustomizableTextObjectFactory(GWarn)
-	{}
-
-	// FCustomizableTextObjectFactory implementation
-	virtual bool CanCreateClass(UClass* InObjectClass, bool& bOmitSubObjs) const override
+	/** Helper to create new components from existing */
+	class FBaseComponentTextFactory
+		: public FCustomizableTextObjectFactory
 	{
-		return InObjectClass->IsChildOf<UDMXPixelMappingBaseComponent>();
-	}
+	public:
 
-	virtual void ProcessConstructedObject(UObject* NewObject) override
-	{
-		check(NewObject);
+		FBaseComponentTextFactory()
+			: FCustomizableTextObjectFactory(GWarn)
+		{}
 
-		if (NewObject->IsA<UDMXPixelMappingBaseComponent>())
+		// FCustomizableTextObjectFactory implementation
+		virtual bool CanCreateClass(UClass* InObjectClass, bool& bOmitSubObjs) const override
 		{
-			UDMXPixelMappingBaseComponent* DMXPixelMappingBaseComponent = Cast<UDMXPixelMappingBaseComponent>(NewObject);
-			DMXPixelMappingBaseComponents.Add(DMXPixelMappingBaseComponent);
+			return InObjectClass->IsChildOf<UDMXPixelMappingBaseComponent>();
 		}
-	}
-
-public:
-	TArray<UDMXPixelMappingBaseComponent*> DMXPixelMappingBaseComponents;
-};
-
-void SDMXPixelMappingHierarchyView::Construct(const FArguments& InArgs, const TSharedPtr<FDMXPixelMappingToolkit>& InToolkit)
-{
-	WeakToolkit = InToolkit;
-
-	bRebuildTreeRequested = false;
-	bIsUpdatingSelection = false;
-
-	SearchBoxWidgetFilter = MakeShared<WidgetTextFilter>(WidgetTextFilter::FItemToStringArray::CreateSP(this, &SDMXPixelMappingHierarchyView::GetWidgetFilterStrings));
-
-	FilterHandler = MakeShared<TreeFilterHandler<TSharedPtr<FDMXPixelMappingHierarchyItemWidgetModel>>>();
-	FilterHandler->SetFilter(SearchBoxWidgetFilter.Get());
-	FilterHandler->SetRootItems(&RootWidgets, &TreeRootWidgets);
-	FilterHandler->SetGetChildrenDelegate(TreeFilterHandler<TSharedPtr<FDMXPixelMappingHierarchyItemWidgetModel>>::FOnGetChildren::CreateRaw(this, &SDMXPixelMappingHierarchyView::WidgetHierarchy_OnGetChildren));
-
-	CommandList = MakeShareable(new FUICommandList);
 
-	CommandList->MapAction(
-		FGenericCommands::Get().Rename,
-		FExecuteAction::CreateSP(this, &SDMXPixelMappingHierarchyView::BeginRename)
-	);
-	CommandList->MapAction(
-		FGenericCommands::Get().Cut,
-		FExecuteAction::CreateSP(this, &SDMXPixelMappingHierarchyView::BeginCut),
-		FCanExecuteAction::CreateSP(this, &SDMXPixelMappingHierarchyView::CanBeginCut)
-	);
-	CommandList->MapAction(
-		FGenericCommands::Get().Copy,
-		FExecuteAction::CreateSP(this, &SDMXPixelMappingHierarchyView::BeginCopy),
-		FCanExecuteAction::CreateSP(this, &SDMXPixelMappingHierarchyView::CanBeginCopy)
-	);
-	CommandList->MapAction(
-		FGenericCommands::Get().Paste,
-		FExecuteAction::CreateSP(this, &SDMXPixelMappingHierarchyView::BeginPaste),
-		FCanExecuteAction::CreateSP(this, &SDMXPixelMappingHierarchyView::CanBeginPaste)
-	);
-	CommandList->MapAction(
-		FGenericCommands::Get().Duplicate,
-		FExecuteAction::CreateSP(this, &SDMXPixelMappingHierarchyView::BeginDuplicate),
-		FCanExecuteAction::CreateSP(this, &SDMXPixelMappingHierarchyView::CanBeginDuplicate)
-	);
-	CommandList->MapAction(
-		FGenericCommands::Get().Delete,
-		FExecuteAction::CreateSP(this, &SDMXPixelMappingHierarchyView::BeginDelete)
-	);
-
-
-	ChildSlot
-		[
-			SNew(SBorder)
-			.BorderImage(FAppStyle::GetBrush("ToolPanel.GroupBorder"))
-			[
-				SNew(SVerticalBox)
-
-				+ SVerticalBox::Slot()
-				.HAlign(HAlign_Left)
-				.Padding(4.f)
-				.AutoHeight()
-				[
-					SNew(SButton)
-					.ButtonStyle(FAppStyle::Get(), "FlatButton.Success")
-					.ForegroundColor(FLinearColor::White)
-					.ToolTipText(LOCTEXT("AddInputSourceTooltip", "Adds a new input source to the pixelmapping asset"))
-					.ContentPadding(FMargin(5.0f, 1.0f))
-					.OnClicked(this, &SDMXPixelMappingHierarchyView::OnAddInputSourceButtonClicked)
-					.Content()
-					[
-						SNew(SHorizontalBox)
-
-						+ SHorizontalBox::Slot()
-						.AutoWidth()
-						.Padding(FMargin(0.f, 1.f))
-						[
-							SNew(SImage)
-							.Image(FAppStyle::GetBrush("Plus"))
-						]
-
-						+ SHorizontalBox::Slot()
-						.VAlign(VAlign_Center)
-						.AutoWidth()
-						.Padding(FMargin(2.f, 0.f, 2.f, 0.f))
-						[
-							SNew(STextBlock)
-							.Text(LOCTEXT("AddInputSourceLabel", "Add Source"))
-						]
-					]
-				]
-
-				+ SVerticalBox::Slot()
-				.FillHeight(1.0f)
-				[
-					SAssignNew(TreeViewArea, SBorder)
-					.Padding(0)
-					.BorderImage(FAppStyle::GetBrush("NoBrush"))
-				]
-			]
-		];
-
-	bRebuildTreeRequested = true;
-
-	GEditor->RegisterForUndo(this);
-
-	// Bind to selection changes
-	InToolkit->GetOnSelectedComponentsChangedDelegate().AddSP(this, &SDMXPixelMappingHierarchyView::OnEditorSelectionChanged);
-
-	// Bind to component changes
-	UDMXPixelMappingBaseComponent::GetOnComponentAdded().AddSP(this, &SDMXPixelMappingHierarchyView::OnComponentAdded);
-	UDMXPixelMappingBaseComponent::GetOnComponentRemoved().AddSP(this, &SDMXPixelMappingHierarchyView::OnComponentRemoved);
-}
-
-void SDMXPixelMappingHierarchyView::Tick(const FGeometry& AllottedGeometry, const double InCurrentTime, const float InDeltaTime)
-{
-	ConditionallyUpdateTree();
-}
-
-FReply SDMXPixelMappingHierarchyView::OnKeyDown(const FGeometry& MyGeometry, const FKeyEvent& InKeyEvent)
-{
-	if (CommandList->ProcessCommandBindings(InKeyEvent))
-	{
-		return FReply::Handled();
-	}
-
-	return FReply::Unhandled();
-}
-
-FReply SDMXPixelMappingHierarchyView::OnAddInputSourceButtonClicked()
-{
-	const TSharedPtr<FDMXPixelMappingToolkit> Toolkit = WeakToolkit.Pin();
-	UDMXPixelMapping* PixelMapping = Toolkit.IsValid() ? Toolkit->GetDMXPixelMapping() : nullptr;
-	UDMXPixelMappingRootComponent* RootComponent = PixelMapping ? PixelMapping->GetRootComponent() : nullptr;
-	if (!Toolkit.IsValid() || !PixelMapping || !RootComponent)
-	{
-		return FReply::Handled();
-	}
-
-	const FScopedTransaction AddRendererTransaction(LOCTEXT("AddRendererTransaction", "Add Pixel Mapping Input Source"));
-	PixelMapping->GetRootComponent()->PreEditChange(nullptr);
-	Toolkit->AddRenderer();
-	PixelMapping->GetRootComponent()->PostEditChange();
-
-	UDMXPixelMappingRendererComponent* RendererComponentToSelect = Toolkit->GetActiveRendererComponent();
-	if (RendererComponentToSelect)
-	{
-		const FDMXPixelMappingComponentReference ComponentRefToSelect(Toolkit, RendererComponentToSelect);
-		const TSet<FDMXPixelMappingComponentReference> NewSelection{ ComponentRefToSelect };
-		Toolkit->SelectComponents(NewSelection);
-
-		// Zoom to fit
-		const TSharedRef<SDMXPixelMappingDesignerView> DesignerView = Toolkit->GetOrCreateDesignerView();
-		const TSharedRef<SDMXPixelMappingPreviewView> PreviewView = Toolkit->GetOrCreatePreviewView();
-		constexpr bool bInstantZoom = false;
-		DesignerView->ZoomToFit(bInstantZoom);
-		PreviewView->ZoomToFit(bInstantZoom);
-	}
-
-	return FReply::Handled();
-}
-
-TSharedPtr<SWidget> SDMXPixelMappingHierarchyView::WidgetHierarchy_OnContextMenuOpening()
-{
-	if (!WeakToolkit.IsValid())
-	{
-		return SNullWidget::NullWidget;
-	}
-
-	FMenuBuilder MenuBuilder(true, CommandList);
-
-	FDMXPixelMappingEditorUtils::CreateComponentContextMenu(MenuBuilder, WeakToolkit.Pin().ToSharedRef());
-
-	return MenuBuilder.MakeWidget();
-}
-
-void SDMXPixelMappingHierarchyView::WidgetHierarchy_OnGetChildren(FDMXPixelMappingHierarchyItemWidgetModelPtr InParent, FDMXPixelMappingHierarchyItemWidgetModelArr& OutChildren)
-{
-	InParent->GatherChildren(OutChildren);
-}
-
-TSharedRef<ITableRow> SDMXPixelMappingHierarchyView::WidgetHierarchy_OnGenerateRow(FDMXPixelMappingHierarchyItemWidgetModelPtr InItem, const TSharedRef<STableViewBase>& OwnerTable)
-{
-	return SNew(SDMXPixelMappingHierarchyItem, OwnerTable, InItem, StaticCastSharedRef<SDMXPixelMappingHierarchyView>(AsShared()));
-}
-
-void SDMXPixelMappingHierarchyView::WidgetHierarchy_OnSelectionChanged(FDMXPixelMappingHierarchyItemWidgetModelPtr SelectedItem, ESelectInfo::Type SelectInfo)
-{
-	if (SelectInfo != ESelectInfo::Direct)
-	{
-		bIsUpdatingSelection = true;
-
-		if (TSharedPtr<FDMXPixelMappingToolkit> Toolkit = WeakToolkit.Pin())
+		virtual void ProcessConstructedObject(UObject* NewObject) override
 		{
-			TSet<FDMXPixelMappingComponentReference> ComponentsToSelect;
-			FDMXPixelMappingHierarchyItemWidgetModelArr SelectedItems = WidgetTreeView->GetSelectedItems();
-			for (FDMXPixelMappingHierarchyItemWidgetModelPtr& Item : SelectedItems)
+			check(NewObject);
+
+			if (NewObject->IsA<UDMXPixelMappingBaseComponent>())
 			{
-				ComponentsToSelect.Add(Item->GetReference());
+				UDMXPixelMappingBaseComponent* DMXPixelMappingBaseComponent = Cast<UDMXPixelMappingBaseComponent>(NewObject);
+				DMXPixelMappingBaseComponents.Add(DMXPixelMappingBaseComponent);
 			}
-
-			Toolkit->SelectComponents(ComponentsToSelect);
 		}
 
-		bIsUpdatingSelection = false;
-	}
-}
+	public:
+		TArray<UDMXPixelMappingBaseComponent*> DMXPixelMappingBaseComponents;
+	};
 
-void SDMXPixelMappingHierarchyView::GetWidgetFilterStrings(FDMXPixelMappingHierarchyItemWidgetModelPtr InModelPtr, TArray<FString>& OutStrings)
-{
-	InModelPtr->GetFilterStrings(OutStrings);
-}
 
-void SDMXPixelMappingHierarchyView::ConditionallyUpdateTree()
-{
+	/** Helper to restore expansion state after refreshing */
 	struct FTreeExpansionSnapshot
 	{
-		static FTreeExpansionSnapshot TakeSnapshot(const FDMXPixelMappingHierarchyItemWidgetModelArr& InRootWidgets, TreeViewPtr InTreeView)
+		static FTreeExpansionSnapshot TakeSnapshot(const FDMXPixelMappingHierarchyItemWidgetModelArr& InRootWidgets, const TSharedPtr<STreeView<FDMXPixelMappingHierarchyItemWidgetModelPtr>>& InTreeView)
 		{
 			FTreeExpansionSnapshot Result;
-			for(FDMXPixelMappingHierarchyItemWidgetModelPtr Model : InRootWidgets)
+			for (FDMXPixelMappingHierarchyItemWidgetModelPtr Model : InRootWidgets)
 			{
 				Result.RecursiveTakeSnapshot(Model, InTreeView);
 			}
 			return Result;
 		}
 
-		void RestoreExpandedAndExpandNewModels(const FDMXPixelMappingHierarchyItemWidgetModelArr& InRootWidgets, TreeViewPtr InTreeView)
+		void RestoreExpandedAndExpandNewModels(const FDMXPixelMappingHierarchyItemWidgetModelArr& InRootWidgets, const TSharedPtr<STreeView<FDMXPixelMappingHierarchyItemWidgetModelPtr>>& InTreeView)
 		{
 			for (FDMXPixelMappingHierarchyItemWidgetModelPtr Model : InRootWidgets)
 			{
@@ -295,413 +95,584 @@ void SDMXPixelMappingHierarchyView::ConditionallyUpdateTree()
 		FTreeExpansionSnapshot()
 		{}
 
-		void RecursiveTakeSnapshot(FDMXPixelMappingHierarchyItemWidgetModelPtr Model, TreeViewPtr TreeView)
+		void RecursiveTakeSnapshot(FDMXPixelMappingHierarchyItemWidgetModelPtr Model, const TSharedPtr<STreeView<FDMXPixelMappingHierarchyItemWidgetModelPtr>>& TreeView)
 		{
-			UDMXPixelMappingBaseComponent* Component = Model->GetReference().GetComponent();
+			UDMXPixelMappingBaseComponent* Component = Model->GetComponent();
 			if (IsValid(Component))
 			{
-			ComponentExpansionStates.Add(Component) = TreeView->IsItemExpanded(Model);
-
-			FDMXPixelMappingHierarchyItemWidgetModelArr Children;
-			Model->GatherChildren(Children);
+				ComponentExpansionStates.Add(Component) = TreeView->IsItemExpanded(Model);
 
-			for (FDMXPixelMappingHierarchyItemWidgetModelPtr& ChildModel : Children)
-			{
-				RecursiveTakeSnapshot(ChildModel, TreeView);
+				for (FDMXPixelMappingHierarchyItemWidgetModelPtr& ChildModel : Model->GetChildren())
+				{
+					RecursiveTakeSnapshot(ChildModel, TreeView);
+				}
 			}
 		}
-		}
 
-		void RecursiveRestoreSnapshot(FDMXPixelMappingHierarchyItemWidgetModelPtr Model, TreeViewPtr TreeView)
+		void RecursiveRestoreSnapshot(FDMXPixelMappingHierarchyItemWidgetModelPtr Model, const TSharedPtr<STreeView<FDMXPixelMappingHierarchyItemWidgetModelPtr>>& TreeView)
 		{
-			UDMXPixelMappingBaseComponent* Component = Model->GetReference().GetComponent();
+			UDMXPixelMappingBaseComponent* Component = Model->GetComponent();
 			if (IsValid(Component))
 			{
-			bool* pPreviousExpansionState = ComponentExpansionStates.Find(Component);
-			if(pPreviousExpansionState == nullptr)
-			{
-				// Initially collapse matrix components
-				if (Cast<UDMXPixelMappingMatrixComponent>(Component))
+				bool* pPreviousExpansionState = ComponentExpansionStates.Find(Component);
+				if (pPreviousExpansionState == nullptr)
 				{
-					TreeView->SetItemExpansion(Model, false);
+					// Initially collapse matrix components
+					if (Cast<UDMXPixelMappingMatrixComponent>(Component))
+					{
+						TreeView->SetItemExpansion(Model, false);
+					}
+					else
+					{
+						TreeView->SetItemExpansion(Model, true);
+					}
 				}
 				else
 				{
-					TreeView->SetItemExpansion(Model, true);
+					TreeView->SetItemExpansion(Model, *pPreviousExpansionState);
 				}
-			}
-			else
-			{
-				TreeView->SetItemExpansion(Model, *pPreviousExpansionState);
-			}
-
-			FDMXPixelMappingHierarchyItemWidgetModelArr Children;
-			Model->GatherChildren(Children);
 
-			for (FDMXPixelMappingHierarchyItemWidgetModelPtr& ChildModel : Children)
-			{
-				RecursiveRestoreSnapshot(ChildModel, TreeView);
+				for (FDMXPixelMappingHierarchyItemWidgetModelPtr& ChildModel : Model->GetChildren())
+				{
+					RecursiveRestoreSnapshot(ChildModel, TreeView);
+				}
 			}
 		}
-		}
-		
-		TMap<UDMXPixelMappingBaseComponent*, bool> ComponentExpansionStates;		
-	};
-	
-	if (bRebuildTreeRequested || bRefreshRequested)
-	{
-		FTreeExpansionSnapshot ExpansionSnapshot = FTreeExpansionSnapshot::TakeSnapshot(RootWidgets, WidgetTreeView);
-		if (bRebuildTreeRequested)
-		{
-			RebuildTreeView();
-		}
 
-		RefreshTree();
-		ExpansionSnapshot.RestoreExpandedAndExpandNewModels(RootWidgets, WidgetTreeView);
-
-		bRefreshRequested = false;
-		bRebuildTreeRequested = false;
+		TMap<UDMXPixelMappingBaseComponent*, bool> ComponentExpansionStates;
+	};
 
-		RestoreSelectedItems();
-	}
 }
 
-void SDMXPixelMappingHierarchyView::RebuildTreeView()
+const FName SDMXPixelMappingHierarchyView::FColumnIds::EditorColor = "EditorColor";
+const FName SDMXPixelMappingHierarchyView::FColumnIds::ComponentName = "Name";
+const FName SDMXPixelMappingHierarchyView::FColumnIds::FixtureID = "FixtureID";
+const FName SDMXPixelMappingHierarchyView::FColumnIds::Patch = "Patch";
+
+void SDMXPixelMappingHierarchyView::Construct(const FArguments& InArgs, const TSharedPtr<FDMXPixelMappingToolkit>& InToolkit)
 {
-	float OldScrollOffset = 0.0f;
+	WeakToolkit = InToolkit;
 
-	if (WidgetTreeView.IsValid())
-	{
-		OldScrollOffset = WidgetTreeView->GetScrollOffset();
-	}
+	CommandList = MakeShareable(new FUICommandList);
+
+	CommandList->MapAction(
+		FGenericCommands::Get().Rename,
+		FExecuteAction::CreateSP(this, &SDMXPixelMappingHierarchyView::RenameSelectedComponent),
+		FCanExecuteAction::CreateSP(this, &SDMXPixelMappingHierarchyView::CanRenameSelectedComponent)
+	);
 
-	SAssignNew(WidgetTreeView, HierarchTreeView)
-		.ItemHeight(20.0f)
-		.SelectionMode(ESelectionMode::Multi)
-		.OnGetChildren(FilterHandler.ToSharedRef(), &TreeFilterHandler<FDMXPixelMappingHierarchyItemWidgetModelPtr>::OnGetFilteredChildren)
-		.OnGenerateRow(this, &SDMXPixelMappingHierarchyView::WidgetHierarchy_OnGenerateRow)
-		.OnSelectionChanged(this, &SDMXPixelMappingHierarchyView::WidgetHierarchy_OnSelectionChanged)
-		.OnContextMenuOpening(this, &SDMXPixelMappingHierarchyView::WidgetHierarchy_OnContextMenuOpening)
-		.TreeItemsSource(&TreeRootWidgets);
+	CommandList->MapAction(
+		FGenericCommands::Get().Delete,
+		FExecuteAction::CreateSP(this, &SDMXPixelMappingHierarchyView::DeleteSelectedComponents)
+	);
 
-	FilterHandler->SetTreeView(WidgetTreeView.Get());
+	using TextFilterType = TTextFilter<FDMXPixelMappingHierarchyItemWidgetModelPtr>;
+	SearchFilter = MakeShared<TextFilterType>(TextFilterType::FItemToStringArray::CreateSP(this, &SDMXPixelMappingHierarchyView::GetWidgetFilterStrings));
 
-	TreeViewArea->SetContent(
-		SNew(SScrollBorder, WidgetTreeView.ToSharedRef())
-		[
-			WidgetTreeView.ToSharedRef()
-		]);
+	FilterHandler = MakeShared<TreeFilterHandler<TSharedPtr<FDMXPixelMappingHierarchyItem>>>();
+	FilterHandler->SetFilter(SearchFilter.Get());
+	FilterHandler->SetRootItems(&AllRootItems, &FilteredRootItems);
+	FilterHandler->SetGetChildrenDelegate(TreeFilterHandler<TSharedPtr<FDMXPixelMappingHierarchyItem>>::FOnGetChildren::CreateSP(this, &SDMXPixelMappingHierarchyView::OnGetChildItems));
 
+	InToolkit->GetOnSelectedComponentsChangedDelegate().AddSP(this, &SDMXPixelMappingHierarchyView::OnEditorSelectionChanged);
+	UDMXPixelMappingBaseComponent::GetOnComponentAdded().AddSP(this, &SDMXPixelMappingHierarchyView::OnComponentAddedOrRemoved);
+	UDMXPixelMappingBaseComponent::GetOnComponentRemoved().AddSP(this, &SDMXPixelMappingHierarchyView::OnComponentAddedOrRemoved);
 
-	WidgetTreeView->SetScrollOffset(OldScrollOffset);
+	BuildChildSlotAndRefresh();
 }
 
-void SDMXPixelMappingHierarchyView::RestoreSelectedItems()
+void SDMXPixelMappingHierarchyView::RequestRefresh()
 {
-	if (bIsUpdatingSelection)
+	if (!RequestRefreshTimerHandle.IsValid())
 	{
-		return;
+		RequestRefreshTimerHandle = GEditor->GetTimerManager()->SetTimerForNextTick(FTimerDelegate::CreateSP(this, &SDMXPixelMappingHierarchyView::ForceRefresh));
 	}
+}
 
-	GEditor->GetTimerManager()->SetTimerForNextTick([this]()
-		{
-			if (WidgetTreeView.IsValid())
-			{
-				WidgetTreeView->ClearSelection();
-			}
-			for (FDMXPixelMappingHierarchyItemWidgetModelPtr& Model : RootWidgets)
-			{
-				RestoreSelectionForItemAndChildren(Model);
-			}
-		});
+void SDMXPixelMappingHierarchyView::ForceRefresh()
+{
+	RequestRefreshTimerHandle.Invalidate();
+
+	using namespace UE::DMX::PixelMappingEditor::SDMXPixelMappingHierarchyView::Private;
+	FTreeExpansionSnapshot ExpansionSnapshot = FTreeExpansionSnapshot::TakeSnapshot(AllRootItems, HierarchyTreeView);
+
+	// Create the root, and let it construct new tree items
+	AllRootItems.Reset();
+	AllRootItems.Add(FDMXPixelMappingHierarchyItem::CreateNew(WeakToolkit.Pin()));
+
+	// Refresh and filter the tree view
+	FilterHandler->RefreshAndFilterTree();
+
+	ExpansionSnapshot.RestoreExpandedAndExpandNewModels(AllRootItems, HierarchyTreeView);
+	AdoptSelectionFromToolkit();
 }
 
-bool SDMXPixelMappingHierarchyView::RestoreSelectionForItemAndChildren(FDMXPixelMappingHierarchyItemWidgetModelPtr& Model)
+void SDMXPixelMappingHierarchyView::BuildChildSlotAndRefresh()
 {
-	if (bIsUpdatingSelection)
-	{
-		return false;
-	}
+	using namespace UE::DMX::PixelMappingEditor::SDMXPixelMappingHierarchyView::Private;
+	FTreeExpansionSnapshot ExpansionSnapshot = FTreeExpansionSnapshot::TakeSnapshot(AllRootItems, HierarchyTreeView);
 
-	const TSharedPtr<FDMXPixelMappingToolkit> Toolkit = WeakToolkit.Pin();
-	if (!Toolkit.IsValid())
-	{
-		return false;
-	}
+	ChildSlot
+		[
+			SNew(SVerticalBox)
 
-	TArray<FDMXPixelMappingHierarchyItemWidgetModelPtr> Children;
-	Model->GatherChildren(Children);
+			+ SVerticalBox::Slot()
+			.AutoHeight()
+			.Padding(4.f)
+			[
+				SNew(SHorizontalBox)
 
-	bool bContainsSelection = false;
-	for (FDMXPixelMappingHierarchyItemWidgetModelPtr& ChildModel : Children)
+				+ SHorizontalBox::Slot()
+				.HAlign(HAlign_Right)
+				.AutoWidth()
+				.Padding(8.f, 0.f, 0.f, 0.f)
+				[
+					SNew(SComboButton)
+					.ComboButtonStyle(FAppStyle::Get(), "SimpleComboButton")
+					.HasDownArrow(true)
+					.OnGetMenuContent(this, &SDMXPixelMappingHierarchyView::GenerateHeaderRowFilterMenu)
+					.ButtonContent()
+					[
+						SNew(SImage)
+						.Image(FAppStyle::Get().GetBrush("Icons.Filter"))
+						.ColorAndOpacity(FSlateColor::UseForeground())
+					]
+				]
+	
+				+ SHorizontalBox::Slot()
+				.FillWidth(1.f)
+				[	
+					SAssignNew(SearchBox, SSearchBox)
+					.OnTextChanged(this, &SDMXPixelMappingHierarchyView::SetFilterText)
+				]
+			]
+
+			+ SVerticalBox::Slot()
+			.FillHeight(1.f)
+			.Padding(4.f)
+			[
+				SAssignNew(HierarchyTreeView, STreeView<FDMXPixelMappingHierarchyItemWidgetModelPtr>)
+				.ItemHeight(20.0f)
+				.SelectionMode(ESelectionMode::Multi)
+				.HeaderRow(GenerateHeaderRow())
+				.OnGetChildren(FilterHandler.ToSharedRef(), &TreeFilterHandler<TSharedPtr<FDMXPixelMappingHierarchyItem>>::OnGetFilteredChildren)
+				.OnGenerateRow(this, &SDMXPixelMappingHierarchyView::OnGenerateRow)
+				.OnSelectionChanged(this, &SDMXPixelMappingHierarchyView::OnSelectionChanged)
+				.OnContextMenuOpening(this, &SDMXPixelMappingHierarchyView::OnContextMenuOpening)
+				.TreeItemsSource(&FilteredRootItems)
+				.OnItemToString_Debug_Lambda([this](FDMXPixelMappingHierarchyItemWidgetModelPtr Item) { return Item->GetComponentNameText().ToString(); })
+			]
+		];
+
+	FilterHandler->SetTreeView(HierarchyTreeView.Get());
+
+	// Create the root, and let it construct new tree items
+	AllRootItems.Reset();
+	AllRootItems.Add(FDMXPixelMappingHierarchyItem::CreateNew(WeakToolkit.Pin()));
+
+	// Refresh and filter the tree view
+	FilterHandler->RefreshAndFilterTree();
+
+	ExpansionSnapshot.RestoreExpandedAndExpandNewModels(AllRootItems, HierarchyTreeView);
+	AdoptSelectionFromToolkit();
+}
+
+TSharedRef<SHeaderRow> SDMXPixelMappingHierarchyView::GenerateHeaderRow()
+{
+	const TSharedRef<SHeaderRow> HeaderRow = SNew(SHeaderRow);
+	const FDMXPixelMappingHierarchySettings& HierarchySettings = GetDefault<UDMXPixelMappingEditorSettings>()->HierarchySettings;
+
+	if (HierarchySettings.bShowEditorColorColumn)
 	{
-		bContainsSelection |= RestoreSelectionForItemAndChildren(ChildModel);
+		HeaderRow->AddColumn(
+			SHeaderRow::FColumn::FArguments()
+			.ColumnId(FColumnIds::EditorColor)
+			.DefaultLabel(LOCTEXT("EditorColorColumnLabel", ""))
+			.FixedWidth(16.f)
+			.VAlignHeader(VAlign_Center)
+		);
 	}
+	
+	HeaderRow->AddColumn(
+		SHeaderRow::FColumn::FArguments()
+		.ColumnId(FColumnIds::ComponentName)
+		.SortMode(this, &SDMXPixelMappingHierarchyView::GetColumnSortMode, FColumnIds::ComponentName)
+		.OnSort(this, &SDMXPixelMappingHierarchyView::SetSortAndRefresh)
+		.FillWidth(0.68f)
+		.HeaderContentPadding(FMargin(6.f))
+		.VAlignHeader(VAlign_Center)
+		.HeaderContent()
+		[
+			SNew(STextBlock)
+			.Text(LOCTEXT("FixturePatchNameColumnLabel", "Name"))
+			.Font(FAppStyle::GetFontStyle(TEXT("PropertyWindow.NormalFont")))
+		]
+	);
 
-	const TSet<FDMXPixelMappingComponentReference> SelectedComponents = Toolkit->GetSelectedComponents();
-	if (bContainsSelection)
+	if (HierarchySettings.bShowFixtureIDColumn)
 	{
-		WidgetTreeView->SetItemExpansion(Model, true);
+		HeaderRow->AddColumn(
+			SHeaderRow::FColumn::FArguments()
+			.ColumnId(FColumnIds::FixtureID)
+			.SortMode(this, &SDMXPixelMappingHierarchyView::GetColumnSortMode, FColumnIds::FixtureID)
+			.OnSort(this, &SDMXPixelMappingHierarchyView::SetSortAndRefresh)
+			.FillWidth(0.16f)
+			.HeaderContentPadding(FMargin(6.f))
+			.VAlignHeader(VAlign_Center)
+			.HeaderContent()
+			[
+				SNew(STextBlock)
+				.Text(LOCTEXT("FixtureIDColumnLabel", "FID"))
+				.Font(FAppStyle::GetFontStyle(TEXT("PropertyWindow.NormalFont")))
+			]
+		);
 	}
 
-	if (SelectedComponents.Contains(Model->GetReference()))
+	if (HierarchySettings.bShowPatchColumn)
 	{
-		WidgetTreeView->SetItemSelection(Model, true, ESelectInfo::Direct);
-		WidgetTreeView->RequestScrollIntoView(Model);
-		
-		return true;
+		HeaderRow->AddColumn(
+			SHeaderRow::FColumn::FArguments()
+			.ColumnId(FColumnIds::Patch)
+			.SortMode(this, &SDMXPixelMappingHierarchyView::GetColumnSortMode, FColumnIds::Patch)
+			.OnSort(this, &SDMXPixelMappingHierarchyView::SetSortAndRefresh)
+			.FillWidth(0.16f)
+			.HeaderContentPadding(FMargin(6.f))
+			.VAlignHeader(VAlign_Center)
+			.HeaderContent()
+			[
+				SNew(STextBlock)
+				.Text(LOCTEXT("PatchColumnLabel", "Patch"))
+				.Font(FAppStyle::GetFontStyle(TEXT("PropertyWindow.NormalFont")))
+			]
+		);
 	}
 
-	return bContainsSelection;
+	return HeaderRow;
 }
 
-void SDMXPixelMappingHierarchyView::RefreshTree()
+TSharedRef<SWidget> SDMXPixelMappingHierarchyView::GenerateHeaderRowFilterMenu()
 {
-	if (!WeakToolkit.IsValid())
+	constexpr bool bShouldCloseMenuAfterSelection = false;
+	FMenuBuilder MenuBuilder(bShouldCloseMenuAfterSelection, nullptr);
+
+	MenuBuilder.BeginSection(NAME_None, LOCTEXT("FilterSection", "Columns"));
 	{
-		return;
+		auto AddMenuEntryLambda = [this, &MenuBuilder](const FText& Label, const FText& ToolTip, const FName& ColumnID)
+		{
+			MenuBuilder.AddMenuEntry(
+				Label,
+				ToolTip,
+				FSlateIcon(),
+				FUIAction(
+					FExecuteAction::CreateSP(this, &SDMXPixelMappingHierarchyView::ToggleColumnVisility, ColumnID),
+					FCanExecuteAction(),
+					FIsActionChecked::CreateSP(this, &SDMXPixelMappingHierarchyView::IsColumVisible, ColumnID)
+				),
+				NAME_None,
+				EUserInterfaceActionType::ToggleButton
+			);
+		};
+
+		AddMenuEntryLambda(
+			LOCTEXT("FixturePatchNameColumn_Label", "Show Editor Color"),
+			FText::GetEmpty(),
+			FColumnIds::EditorColor
+		);
+
+		AddMenuEntryLambda(
+			LOCTEXT("FixturePatchTypeColumn_Label", "Show Fixture ID"),
+			FText::GetEmpty(),
+			FColumnIds::FixtureID
+		);
+
+		AddMenuEntryLambda(
+			LOCTEXT("FixturePatchModeColumn_Label", "Show Patch"),
+			FText::GetEmpty(),
+			FColumnIds::Patch
+		);
 	}
+	MenuBuilder.EndSection();
 
-	RootWidgets.Empty();
-	RootWidgets.Add(MakeShared<FDMXPixelMappingHierarchyItemWidgetModel>(WeakToolkit.Pin()));
-
-	FilterHandler->RefreshAndFilterTree();
-}
-
-void SDMXPixelMappingHierarchyView::OnComponentAdded(UDMXPixelMapping* PixelMapping, UDMXPixelMappingBaseComponent* Component)
-{
-	RequestRebuildTree();
-}
-
-/** Called when a component was removed */
-void SDMXPixelMappingHierarchyView::OnComponentRemoved(UDMXPixelMapping* PixelMapping, UDMXPixelMappingBaseComponent* Component)
-{
-	RequestRebuildTree();
+	return MenuBuilder.MakeWidget();
 }
 
-void SDMXPixelMappingHierarchyView::OnEditorSelectionChanged()
+FReply SDMXPixelMappingHierarchyView::OnKeyDown(const FGeometry& MyGeometry, const FKeyEvent& InKeyEvent)
 {
-	if (!bIsUpdatingSelection)
+	if (CommandList->ProcessCommandBindings(InKeyEvent))
 	{
-		RestoreSelectedItems();
+		return FReply::Handled();
 	}
-}
 
-void SDMXPixelMappingHierarchyView::BeginRename()
-{
-	FDMXPixelMappingHierarchyItemWidgetModelArr SelectedItems = WidgetTreeView->GetSelectedItems();
-	SelectedItems[0]->RequestBeginRename();
+	return FReply::Unhandled();
 }
 
-bool SDMXPixelMappingHierarchyView::CanBeginCut() const
+void SDMXPixelMappingHierarchyView::PostUndo(bool bSuccess)
 {
-	return CanBeginCopy();
+	RequestRefresh();
 }
 
-void SDMXPixelMappingHierarchyView::BeginCut()
-{
-	BeginCopy();
-	BeginDelete();
+void SDMXPixelMappingHierarchyView::PostRedo(bool bSuccess)
+{ 
+	RequestRefresh();
 }
 
-bool SDMXPixelMappingHierarchyView::CanBeginCopy() const
+void SDMXPixelMappingHierarchyView::OnGetChildItems(TSharedPtr<FDMXPixelMappingHierarchyItem> InParent, TArray<TSharedPtr<FDMXPixelMappingHierarchyItem>>& OutChildren)
 {
-	FDMXPixelMappingHierarchyItemWidgetModelArr SelectedItems = WidgetTreeView->GetSelectedItems();
-	int32 NumberOfValidItems = 0;
-	for (FDMXPixelMappingHierarchyItemWidgetModelPtr SelectedItem : SelectedItems)
+	// If the parent is a fixture group, sort its children using the current sort order
+	if (InParent.IsValid() && 
+		InParent->GetComponent() && 
+		InParent->GetComponent()->GetClass() == UDMXPixelMappingFixtureGroupComponent::StaticClass())
 	{
-		UDMXPixelMappingBaseComponent* SelectedComponent = SelectedItem->GetReference().GetComponent();
-		// root component cannot be copied
-		if (!SelectedComponent->IsA<UDMXPixelMappingRootComponent>())
+		const FDMXPixelMappingHierarchySettings& HierarchySettings = GetDefault<UDMXPixelMappingEditorSettings>()->HierarchySettings;
+		if (HierarchySettings.SortByColumnId == FColumnIds::ComponentName)
 		{
-			NumberOfValidItems++;
+			InParent->StableSortChildren([](const TSharedPtr<FDMXPixelMappingHierarchyItem>& Item)
+				{
+					return Item->GetComponentNameText().ToString();
+				});
+		}
+		else if (HierarchySettings.SortByColumnId == FColumnIds::Patch)
+		{
+			InParent->StableSortChildren([](const TSharedPtr<FDMXPixelMappingHierarchyItem>& Item)
+				{
+					return Item->GetAbsoluteChannel();
+				});
+		}
+
+		if (!HierarchySettings.bSortAscending)
+		{
+			InParent->ReverseChildren();
 		}
 	}
 
-	return NumberOfValidItems > 0;
+	OutChildren = InParent->GetChildren();
 }
 
-void SDMXPixelMappingHierarchyView::BeginCopy()
+TSharedRef<ITableRow> SDMXPixelMappingHierarchyView::OnGenerateRow(FDMXPixelMappingHierarchyItemWidgetModelPtr Item, const TSharedRef<STableViewBase>& OwnerTable)
 {
-	const FExportObjectInnerContext Context;
-	FStringOutputDevice Archive;
+	return SNew(SDMXPixelMappingHierarchyItem, OwnerTable, WeakToolkit, Item.ToSharedRef());
+}
 
-	FDMXPixelMappingHierarchyItemWidgetModelArr SelectedItems = WidgetTreeView->GetSelectedItems();
+TSharedPtr<SWidget> SDMXPixelMappingHierarchyView::OnContextMenuOpening()
+{
+	if (!WeakToolkit.IsValid())
+	{
+		return SNullWidget::NullWidget;
+	}
 
-	for (FDMXPixelMappingHierarchyItemWidgetModelPtr SelectedItem : SelectedItems)
+	FMenuBuilder MenuBuilder(true, CommandList);
+
+	MenuBuilder.BeginSection("Edit", LOCTEXT("Edit", "Edit"));
 	{
-		UDMXPixelMappingBaseComponent* SelectedComponent = SelectedItem->GetReference().GetComponent();
-		if (!SelectedComponent->IsA<UDMXPixelMappingRootComponent>())
-		{
-			UExporter::ExportToOutputDevice(&Context, SelectedComponent, nullptr, Archive, TEXT("copy"), 0, PPF_ExportsNotFullyQualified | PPF_Copy | PPF_Delimited, false, GetTransientPackage());
-		}
+		MenuBuilder.AddMenuEntry(FGenericCommands::Get().Rename);
+		MenuBuilder.AddMenuEntry(FGenericCommands::Get().Delete);
 	}
+	MenuBuilder.EndSection();
 
-	// Copy text to clipboard
-	FString ExportedText = Archive;
-	FPlatformApplicationMisc::ClipboardCopy(*ExportedText);
+	return MenuBuilder.MakeWidget();
 }
 
-bool SDMXPixelMappingHierarchyView::CanBeginPaste() const
+void SDMXPixelMappingHierarchyView::OnSelectionChanged(FDMXPixelMappingHierarchyItemWidgetModelPtr SelectedItem, ESelectInfo::Type SelectInfo)
 {
-	FDMXPixelMappingHierarchyItemWidgetModelArr SelectedItems = WidgetTreeView->GetSelectedItems();
-
-	TSet<TSubclassOf<UDMXPixelMappingBaseComponent>> FoundComponentsType;
-	for (FDMXPixelMappingHierarchyItemWidgetModelPtr SelectedItem : SelectedItems)
+	if (SelectInfo == ESelectInfo::Direct)
 	{
-		UDMXPixelMappingBaseComponent* SelectedComponent = SelectedItem->GetReference().GetComponent();
-		FoundComponentsType.Add(SelectedComponent->GetClass());
+		return;
 	}
 
-	// only one destination type is allowed (read: two renderers as well as two matrix are allowed)
-	return FoundComponentsType.Num() < 2;
-}
+	const TGuardValue<bool> GuardIsUpdatingSelection(bIsUpdatingSelection, true);
 
-void SDMXPixelMappingHierarchyView::RecursivePaste(UDMXPixelMappingBaseComponent* InComponent)
-{
-	for (int32 ChildIndex = 0; ChildIndex < InComponent->Children.Num(); ChildIndex++)
+	if (TSharedPtr<FDMXPixelMappingToolkit> Toolkit = WeakToolkit.Pin())
 	{
-		UDMXPixelMappingBaseComponent* ChildComponent = InComponent->Children[ChildIndex];
-		if (ChildComponent)
+		TSet<FDMXPixelMappingComponentReference> ComponentsToSelect;
+		FDMXPixelMappingHierarchyItemWidgetModelArr SelectedItems = HierarchyTreeView->GetSelectedItems();
+		for (FDMXPixelMappingHierarchyItemWidgetModelPtr& Item : SelectedItems)
 		{
-			const FName UniqueName = MakeUniqueObjectName(ChildComponent->GetOuter(), ChildComponent->GetClass(), ChildComponent->GetFName());
-			UDMXPixelMappingBaseComponent* ChildCopy = DuplicateObject(ChildComponent, ChildComponent->GetOuter(), UniqueName);
-			if (ChildCopy)
-			{
-				ChildCopy->SetParent(InComponent);
-				InComponent->Children[ChildIndex] = ChildCopy;
-				RecursivePaste(InComponent->Children[ChildIndex]);
-			}
+			ComponentsToSelect.Add(FDMXPixelMappingComponentReference(Toolkit, Item->GetComponent()));
 		}
+
+		Toolkit->SelectComponents(ComponentsToSelect);
 	}
 }
 
-void SDMXPixelMappingHierarchyView::RequestRebuildTree()
+void SDMXPixelMappingHierarchyView::ToggleColumnVisility(FName ColumnId)
 {
-	bRebuildTreeRequested = true;
+	UDMXPixelMappingEditorSettings* EditorSettings = GetMutableDefault<UDMXPixelMappingEditorSettings>();
+	if (ColumnId == FColumnIds::EditorColor)
+	{
+		EditorSettings->HierarchySettings.bShowEditorColorColumn = !EditorSettings->HierarchySettings.bShowEditorColorColumn;
+	}
+	else if (ColumnId == FColumnIds::FixtureID)
+	{
+		EditorSettings->HierarchySettings.bShowFixtureIDColumn = !EditorSettings->HierarchySettings.bShowFixtureIDColumn;
+	}
+	else if (ColumnId == FColumnIds::Patch)
+	{
+		EditorSettings->HierarchySettings.bShowPatchColumn = !EditorSettings->HierarchySettings.bShowPatchColumn;
+	}
+	EditorSettings->SaveConfig();
+
+	// To adopt column changes, the entire child slot needs to be updated
+	BuildChildSlotAndRefresh();
 }
 
-bool SDMXPixelMappingHierarchyView::MoveComponentToComponent(UDMXPixelMappingBaseComponent* Source, UDMXPixelMappingBaseComponent* Destination, const bool bRename)
+bool SDMXPixelMappingHierarchyView::IsColumVisible(FName ColumnId) const
 {
-	UDMXPixelMappingBaseComponent* NewParent = nullptr;
-	// first check if we can directly move the component
-	if (Source->CanBeMovedTo(Destination))
+	UDMXPixelMappingEditorSettings* EditorSettings = GetMutableDefault<UDMXPixelMappingEditorSettings>();
+	if (ColumnId == FColumnIds::EditorColor)
 	{
-		NewParent = Destination;
+		return EditorSettings->HierarchySettings.bShowEditorColorColumn;
 	}
-	// try parent
-	else
+	else if (ColumnId == FColumnIds::FixtureID)
 	{
-		if (Source->CanBeMovedTo(Destination->GetParent()))
-		{
-			NewParent = Destination->GetParent();
-		}
+		return EditorSettings->HierarchySettings.bShowFixtureIDColumn;
 	}
-
-	// pasting while multiple components are selected, can
-	// cause the same children to be attached to the same parent
-	// we need to check it
-	if (NewParent != nullptr && !NewParent->Children.Contains(Source))
+	else if (ColumnId == FColumnIds::Patch)
 	{
-		NewParent->Modify();
-		Source->Modify();
-
-		RecursivePaste(Source);
+		return EditorSettings->HierarchySettings.bShowPatchColumn;
+	}
 
-		if (bRename)
-		{
-			const FName UniqueName = MakeUniqueObjectName(NewParent->GetOuter(), Source->GetClass(), Source->GetFName());
-			const FString NewNameStr = UniqueName.ToString();
+	return false;
+}
 
-			Source->Rename(*NewNameStr, NewParent->GetOuter());
-		}
+EColumnSortMode::Type SDMXPixelMappingHierarchyView::GetColumnSortMode(FName ColumnId) const
+{
+	const FDMXPixelMappingHierarchySettings& HierarchySettings = GetDefault<UDMXPixelMappingEditorSettings>()->HierarchySettings;
+	if (HierarchySettings.SortByColumnId != ColumnId)
+	{
+		return EColumnSortMode::None;
+	}
 
-		NewParent->AddChild(Source);
+	return HierarchySettings.bSortAscending ? EColumnSortMode::Ascending : EColumnSortMode::Descending;
+}
 
-		RequestComponentRedraw(NewParent);
+void SDMXPixelMappingHierarchyView::SetSortAndRefresh(const EColumnSortPriority::Type SortPriority, const FName& ColumnId, const EColumnSortMode::Type InSortMode)
+{
+	UDMXPixelMappingEditorSettings* EditorSettings = GetMutableDefault<UDMXPixelMappingEditorSettings>();
 
-		RequestRebuildTree();
+	EditorSettings->HierarchySettings.bSortAscending = InSortMode == EColumnSortMode::Ascending;
+	EditorSettings->HierarchySettings.SortByColumnId = ColumnId;
+	EditorSettings->SaveConfig();
 
-		return true;
-	}
+	RequestRefresh();
+}
 
-	return false;
+void SDMXPixelMappingHierarchyView::OnComponentAddedOrRemoved(UDMXPixelMapping* PixelMapping, UDMXPixelMappingBaseComponent* Component)
+{
+	RequestRefresh();
 }
 
-void SDMXPixelMappingHierarchyView::BeginPaste()
+void SDMXPixelMappingHierarchyView::OnEditorSelectionChanged()
 {
-	FString TextToImport;
-	FPlatformApplicationMisc::ClipboardPaste(TextToImport);
+	if (!bIsUpdatingSelection)
+	{
+		AdoptSelectionFromToolkit();
+	}
+}
 
-	FDMXPixelMappingBaseComponentTextFactory Factory;
-	Factory.ProcessBuffer(GetTransientPackage(), RF_Transactional, TextToImport);
-	TArray<UDMXPixelMappingBaseComponent*> NewComponents = Factory.DMXPixelMappingBaseComponents;
+void SDMXPixelMappingHierarchyView::AdoptSelectionFromToolkit()
+{
+	if (bIsUpdatingSelection)
+	{
+		return;
+	}
 
-	FDMXPixelMappingHierarchyItemWidgetModelArr SelectedItems = WidgetTreeView->GetSelectedItems();
-	const FScopedTransaction Transaction(FText::Format(LOCTEXT("PasteComponent", "Paste {0}|plural(one=Component, other=Components)"), SelectedItems.Num()));
+	if (HierarchyTreeView.IsValid())
+	{
+		HierarchyTreeView->ClearSelection();
+	}
 
-	for (FDMXPixelMappingHierarchyItemWidgetModelPtr SelectedItem : SelectedItems)
+	for (const TSharedPtr<FDMXPixelMappingHierarchyItem>& Item : FilteredRootItems)
 	{
-		UDMXPixelMappingBaseComponent* SelectedComponent = SelectedItem->GetReference().GetComponent();
-		for (UDMXPixelMappingBaseComponent* Component : NewComponents)
-		{
-			MoveComponentToComponent(Component, SelectedComponent, true);
-		}
+		RecursiveAdoptSelectionFromToolkit(Item.ToSharedRef());
 	}
 }
 
-bool SDMXPixelMappingHierarchyView::CanBeginDuplicate() const
+bool SDMXPixelMappingHierarchyView::RecursiveAdoptSelectionFromToolkit(const TSharedRef<FDMXPixelMappingHierarchyItem>& Item)
 {
-	return CanBeginCopy();
-}
+	if (bIsUpdatingSelection)
+	{
+		return false;
+	}
 
-void SDMXPixelMappingHierarchyView::BeginDuplicate()
-{
-	FDMXPixelMappingHierarchyItemWidgetModelArr SelectedItems = WidgetTreeView->GetSelectedItems();
+	const TSharedPtr<FDMXPixelMappingToolkit> Toolkit = WeakToolkit.Pin();
+	if (!Toolkit.IsValid())
+	{
+		return false;
+	}
 
-	const FScopedTransaction Transaction(FText::Format(LOCTEXT("DMXPixelMapping.DuplicateComponents", "Duplicate {0}|plural(one=Component, other=Components)"), SelectedItems.Num()));
+	bool bContainsSelection = false;
+	for (const TSharedPtr<FDMXPixelMappingHierarchyItem>& ChildItem : Item->GetChildren())
+	{
+		bContainsSelection |= RecursiveAdoptSelectionFromToolkit(ChildItem.ToSharedRef());
+	}
 
-	for (FDMXPixelMappingHierarchyItemWidgetModelPtr SelectedItem : SelectedItems)
+	const TSet<FDMXPixelMappingComponentReference> SelectedComponents = Toolkit->GetSelectedComponents();
+	if (bContainsSelection)
 	{
-		UDMXPixelMappingBaseComponent* SelectedComponent = SelectedItem->GetReference().GetComponent();
-		const FName UniqueName = MakeUniqueObjectName(SelectedComponent->GetOuter(), SelectedComponent->GetClass(), SelectedComponent->GetFName());
-		UDMXPixelMappingBaseComponent* NewComponent = DuplicateObject(SelectedComponent, SelectedComponent->GetOuter(), UniqueName);
-		if (NewComponent)
-		{
-			MoveComponentToComponent(NewComponent, SelectedComponent, false);
-		}
+		HierarchyTreeView->SetItemExpansion(Item, true);
 	}
+
+	if (SelectedComponents.Contains(FDMXPixelMappingComponentReference(Toolkit, Item->GetComponent())))
+	{
+		HierarchyTreeView->SetItemSelection(Item, true, ESelectInfo::Direct);
+		HierarchyTreeView->RequestScrollIntoView(Item);
+
+		return true;
+	}
+
+	return bContainsSelection;
 }
 
-void SDMXPixelMappingHierarchyView::RequestComponentRedraw(UDMXPixelMappingBaseComponent* Component)
+bool SDMXPixelMappingHierarchyView::CanRenameSelectedComponent() const
 {
-	if (!Component)
+	return HierarchyTreeView.IsValid() ? HierarchyTreeView->GetSelectedItems().Num() == 1 : false;
+}
+
+void SDMXPixelMappingHierarchyView::RenameSelectedComponent()
+{
+	const TArray<TSharedPtr<FDMXPixelMappingHierarchyItem>> SelectedItems = HierarchyTreeView.IsValid() ? HierarchyTreeView->GetSelectedItems() : TArray<TSharedPtr<FDMXPixelMappingHierarchyItem>>();
+	if (!ensureMsgf(SelectedItems.Num() == 1, TEXT("Cannot rename selected components. Please call CanRenameSelectedComponent() before calling RenameSelectedComponent().")))
 	{
 		return;
 	}
 
-	if (TSharedPtr<FDMXPixelMappingToolkit> Toolkit = WeakToolkit.Pin())
+	const TSharedPtr<ITableRow> TableRow = HierarchyTreeView->WidgetFromItem(SelectedItems[0]);
+	if (!ensureMsgf(TableRow.IsValid(), TEXT("Cannot find widget for item. Cannot rename component")))
 	{
-		TSet<FDMXPixelMappingComponentReference> ComponentsToUpdate;
-		FDMXPixelMappingComponentReference ComponentReference(Toolkit, Component);
-		ComponentsToUpdate.Add(ComponentReference);
+		return;
 	}
+
+	StaticCastSharedPtr<SDMXPixelMappingHierarchyItem>(TableRow)->EnterRenameMode();
 }
 
-void SDMXPixelMappingHierarchyView::BeginDelete()
+void SDMXPixelMappingHierarchyView::DeleteSelectedComponents()
 {
 	if (TSharedPtr<FDMXPixelMappingToolkit> Toolkit = WeakToolkit.Pin())
 	{
 		const int32 NumSelectedComponents = Toolkit->GetSelectedComponents().Num();
 		const FScopedTransaction Transaction(FText::Format(LOCTEXT("DMXPixelMapping.RemoveComponents", "Remove {0}|plural(one=Component, other=Components)"), NumSelectedComponents));
-		
+
 		Toolkit->DeleteSelectedComponents();
 	}
 }
 
-void SDMXPixelMappingHierarchyView::PostUndo(bool bSuccess)
+void SDMXPixelMappingHierarchyView::SetFilterText(const FText& Text)
+{
+	using namespace UE::DMX::PixelMappingEditor::SDMXPixelMappingHierarchyView::Private;
+	FTreeExpansionSnapshot ExpansionSnapshot = FTreeExpansionSnapshot::TakeSnapshot(AllRootItems, HierarchyTreeView);
+
+	FilterHandler->SetIsEnabled(!Text.IsEmpty());
+	SearchFilter->SetRawFilterText(Text);
+	FilterHandler->RefreshAndFilterTree();
+
+	ExpansionSnapshot.RestoreExpandedAndExpandNewModels(AllRootItems, HierarchyTreeView);
+}
+
+void SDMXPixelMappingHierarchyView::GetWidgetFilterStrings(FDMXPixelMappingHierarchyItemWidgetModelPtr InModel, TArray<FString>& OutStrings) const
 {
-	RequestRebuildTree();
+	OutStrings.Add(InModel->GetComponentNameText().ToString());
+	OutStrings.Add(InModel->GetFixtureIDText().ToString());
+	OutStrings.Add(InModel->GetPatchText().ToString());
 }
 
 #undef LOCTEXT_NAMESPACE
diff --git a/Engine/Plugins/VirtualProduction/DMX/DMXPixelMapping/Source/DMXPixelMappingEditor/Private/Views/SDMXPixelMappingHierarchyView.h b/Engine/Plugins/VirtualProduction/DMX/DMXPixelMapping/Source/DMXPixelMappingEditor/Private/Views/SDMXPixelMappingHierarchyView.h
--- a/Engine/Plugins/VirtualProduction/DMX/DMXPixelMapping/Source/DMXPixelMappingEditor/Private/Views/SDMXPixelMappingHierarchyView.h
+++ b/Engine/Plugins/VirtualProduction/DMX/DMXPixelMapping/Source/DMXPixelMappingEditor/Private/Views/SDMXPixelMappingHierarchyView.h
@@ -3,48 +3,44 @@
 #pragma once
 
 #include "DMXPixelMappingEditorCommon.h"
-
 #include "EditorUndoClient.h"
+#include "Engine/EngineTypes.h"
 #include "Widgets/SCompoundWidget.h"
+#include "Widgets/Views/SHeaderRow.h"
 
-class UDMXPixelMappingBaseComponent;
-namespace ESelectInfo { enum Type : int; }
-template <typename ItemType> class TTextFilter;
-
+class FDMXPixelMappingHierarchyItem;
 class FDMXPixelMappingToolkit;
-class UDMXPixelMapping;
-
 class FUICommandList;
 class ITableRow;
 class SBorder;
+class SHeaderRow;
 class SSearchBox;
 class STableViewBase;
 template<typename ItemType> class STreeView;
 template<typename ItemType> class TreeFilterHandler;
+template <typename ItemType> class TTextFilter;
+class UDMXPixelMapping;
+class UDMXPixelMappingBaseComponent;
 
 
-class SDMXPixelMappingHierarchyView
-	: public SCompoundWidget, public FEditorUndoClient
-{
-public:
-	using WidgetTextFilter = TTextFilter<FDMXPixelMappingHierarchyItemWidgetModelPtr>;
-	using HierarchTreeView = STreeView<FDMXPixelMappingHierarchyItemWidgetModelPtr>;
-	using TreeViewPtr = TSharedPtr<HierarchTreeView>;
-
-	enum class EExpandBehavior : uint8
-	{
-		NeverExpand,
-		AlwaysExpand,
-		RestoreFromPrevious,
-		FromModel
-	};
 
+class SDMXPixelMappingHierarchyView final
+	: public SCompoundWidget
+	, public FSelfRegisteringEditorUndoClient 
+{
 public:
-
 	SLATE_BEGIN_ARGS(SDMXPixelMappingHierarchyView) { }
 	SLATE_END_ARGS()
+		
+	/** ColumnIds for the tree view */
+	struct FColumnIds
+	{
+		static const FName EditorColor;
+		static const FName ComponentName;
+		static const FName FixtureID;
+		static const FName Patch;
+	};
 
-public:
 	/**
 	 * Constructs the widget.
 	 *
@@ -52,110 +48,115 @@ public:
 	 */
 	void Construct(const FArguments& InArgs, const TSharedPtr<FDMXPixelMappingToolkit>& InToolkit);
 
-	// Begin SWidget
-	virtual void Tick(const FGeometry& AllottedGeometry, const double InCurrentTime, const float InDeltaTime) override;
-	virtual FReply OnKeyDown(const FGeometry& MyGeometry, const FKeyEvent& InKeyEvent) override;
-	// End SWidget
-
-	 //~ Begin FEditorUndoClient Interface
-	virtual void PostUndo(bool bSuccess) override;
-	virtual void PostRedo(bool bSuccess) override { PostUndo(bSuccess); }
-	// End of FEditorUndoClient
+	/** Refreshes the widget on the next tick */
+	void RequestRefresh();
 
-	/** force a rebuild of the hierarchy **/
-	void RequestRebuildTree();
+private:
+	/** Refreshes the widget */
+	void ForceRefresh();
 
-	/** force a redraw of a component **/
-	void RequestComponentRedraw(UDMXPixelMappingBaseComponent* Component);
+	/** Builds the child slot anew and refreshes the widget */
+	void BuildChildSlotAndRefresh();
 
-private:
-	/** Called when the 'Add Input Source' button was clicked */
-	FReply OnAddInputSourceButtonClicked();
+	/** Generates the header row for the tree view */
+	TSharedRef<SHeaderRow> GenerateHeaderRow();
 
-	TSharedPtr<SWidget> WidgetHierarchy_OnContextMenuOpening();
-	void WidgetHierarchy_OnGetChildren(FDMXPixelMappingHierarchyItemWidgetModelPtr InParent, FDMXPixelMappingHierarchyItemWidgetModelArr& OutChildren);
-	TSharedRef< ITableRow > WidgetHierarchy_OnGenerateRow(FDMXPixelMappingHierarchyItemWidgetModelPtr InItem, const TSharedRef<STableViewBase>& OwnerTable);
-	void WidgetHierarchy_OnSelectionChanged(FDMXPixelMappingHierarchyItemWidgetModelPtr SelectedItem, ESelectInfo::Type SelectInfo);
+	/** Generates the header row filter menu */
+	TSharedRef<SWidget> GenerateHeaderRowFilterMenu();
+	
+	//~ Begin SWidget interface
+	virtual FReply OnKeyDown(const FGeometry& MyGeometry, const FKeyEvent& InKeyEvent) override;
+	//~ End SWidget interface
 
-	/**  Gets an array of strings used for filtering/searching the specified widget. */
-	void GetWidgetFilterStrings(FDMXPixelMappingHierarchyItemWidgetModelPtr InModelPtr, TArray<FString>& OutStrings);
+	//~ Begin FEditorUndoClient interface
+	virtual void PostUndo(bool bSuccess) override;
+	virtual void PostRedo(bool bSuccess) override;
+	//~ End FEditorUndoClient interface
 
-	void ConditionallyUpdateTree();
-	/** Completely regenerates the treeview */
-	void RebuildTreeView();
+	/** Called to get child items */
+	void OnGetChildItems(FDMXPixelMappingHierarchyItemWidgetModelPtr InParent, FDMXPixelMappingHierarchyItemWidgetModelArr& OutChildren);
 
-	void RestoreSelectedItems();
+	/** Called when a row is generated */
+	TSharedRef<ITableRow> OnGenerateRow(FDMXPixelMappingHierarchyItemWidgetModelPtr Item, const TSharedRef<STableViewBase>& OwnerTable);
 
-	/** Restores selection for a the item and its children. Returns true if the Model or a child was selected. */
-	bool RestoreSelectionForItemAndChildren(FDMXPixelMappingHierarchyItemWidgetModelPtr& Model);
+	/** Called when the context menu is opening */
+	TSharedPtr<SWidget> OnContextMenuOpening();
 
-	/** Rebuilds the tree structure based on the current filter options */
-	void RefreshTree();
+	/** Called when selection changed */
+	void OnSelectionChanged(FDMXPixelMappingHierarchyItemWidgetModelPtr SelectedItem, ESelectInfo::Type SelectInfo);
 
-	/** Called when a component was added */
-	void OnComponentAdded(UDMXPixelMapping* PixelMapping, UDMXPixelMappingBaseComponent* Component);
+	/** Toggles the visibility of a column */
+	void ToggleColumnVisility(FName ColumnId);
 
-	/** Called when a component was removed */
-	void OnComponentRemoved(UDMXPixelMapping* PixelMapping, UDMXPixelMappingBaseComponent* Component);
+	/** Returns true if the specified column is currently visible */
+	bool IsColumVisible(FName ColumnId) const;
 
-	/** Called when the selected widget has changed.  The treeview then needs to match the new selection. */
-	void OnEditorSelectionChanged();
+	/** Gets the sort mode for the given column ID */
+	EColumnSortMode::Type GetColumnSortMode(FName ColumnId) const;
 
-	void BeginRename();
+	/** Set the current sorting state of the hierarchy tree view and refreshes the tree, applying the new sorting */
+	void SetSortAndRefresh(const EColumnSortPriority::Type SortPriority, const FName& ColumnId, const EColumnSortMode::Type InSortMode);
 
-	void BeginCut();
+	/** Called when a component was added or removed from the pixel mapping */
+	void OnComponentAddedOrRemoved(UDMXPixelMapping* PixelMapping, UDMXPixelMappingBaseComponent* Component);
 
-	bool CanBeginCut() const;
+	/** Called when the selected widget has changed.  The treeview then needs to match the new selection. */
+	void OnEditorSelectionChanged();
 
-	void BeginCopy();
+	/** Adopts the current selection of the toolkit */
+	void AdoptSelectionFromToolkit();
 
-	bool CanBeginCopy() const;
+	/** Adopts the selection from toolkit */
+	bool RecursiveAdoptSelectionFromToolkit(const TSharedRef<FDMXPixelMappingHierarchyItem>& Model);
 
-	void BeginPaste();
+	/** Returns true if the current selection can to be renamed */
+	bool CanRenameSelectedComponent() const;
 
-	bool CanBeginPaste() const;
+	/** Renames selected components */
+	void RenameSelectedComponent();
 
-	void BeginDuplicate();
+	/** Returns true if the current selection can to be deleted */
+	bool CanDeleteSelectedComponents() const;
 
-	bool CanBeginDuplicate() const;
+	/** Deletes selected components */
+	void DeleteSelectedComponents();
 
-	void BeginDelete();
+	/** Sets the text for the filter */
+	void SetFilterText(const FText& Text);
 
-	FText GetSearchText() const { return FText(); }
+	/**  Gets an array of strings used for filtering/searching the specified widget. */
+	void GetWidgetFilterStrings(FDMXPixelMappingHierarchyItemWidgetModelPtr InModel, TArray<FString>& OutStrings) const;
 
-	void RecursivePaste(UDMXPixelMappingBaseComponent* InComponent);
+	/** Flag to ignore selections while the hierarchy view is updating the selection. */
+	bool bIsUpdatingSelection = false;
 
-	bool MoveComponentToComponent(UDMXPixelMappingBaseComponent* Source, UDMXPixelMappingBaseComponent* Destination, const bool bRename);
+	/** All items in the tree */
+	TArray<TSharedPtr<FDMXPixelMappingHierarchyItem>> TreeItems;
 
-private:
-	TWeakPtr<FDMXPixelMappingToolkit> WeakToolkit;
+	/** Root items being displayed */
+	TArray<TSharedPtr<FDMXPixelMappingHierarchyItem>> FilteredRootItems;
 
-	/** Commands specific to the hierarchy. */
-	TSharedPtr<FUICommandList> CommandList;
+	/** All root items */
+	TArray<TSharedPtr<FDMXPixelMappingHierarchyItem>> AllRootItems;
 
-	TSharedPtr<SBorder> TreeViewArea;
+	/** The search box widget */
+	TSharedPtr<SSearchBox> SearchBox;
 
-	TSharedPtr<SSearchBox> SearchBoxPtr;
+	/** The hierarchy tree view widget */
+	TSharedPtr<STreeView<FDMXPixelMappingHierarchyItemWidgetModelPtr>> HierarchyTreeView;
 
 	/** Handles filtering the hierarchy based on an IFilter. */
 	TSharedPtr<TreeFilterHandler<FDMXPixelMappingHierarchyItemWidgetModelPtr>> FilterHandler;
 
-	TSharedPtr<WidgetTextFilter> SearchBoxWidgetFilter;
-
-	/** The source root widgets for the tree. */
-	FDMXPixelMappingHierarchyItemWidgetModelArr RootWidgets;
-
-	/** The root widgets which are actually displayed by the TreeView which will be managed by the TreeFilterHandler. */
-	FDMXPixelMappingHierarchyItemWidgetModelArr TreeRootWidgets;
+	/** Text filter for the filter handler, set from the search bbox */
+	TSharedPtr<TTextFilter<FDMXPixelMappingHierarchyItemWidgetModelPtr>> SearchFilter;
 
-	/** Has a full refresh of the tree been requested?  This happens when the user is filtering the tree */
-	bool bRefreshRequested;
+	/** Timer handle for the request refresh timer */
+	FTimerHandle RequestRefreshTimerHandle;
 
-	/** Is the tree in such a changed state that the whole widget needs rebuilding? */
-	bool bRebuildTreeRequested;
-
-	/** Flag to ignore selections while the hierarchy view is updating the selection. */
-	bool bIsUpdatingSelection;
+	/** Commands specific to the hierarchy. */
+	TSharedPtr<FUICommandList> CommandList;
 
-	TreeViewPtr WidgetTreeView;
+	/** The toolkit that contains this widget */
+	TWeakPtr<FDMXPixelMappingToolkit> WeakToolkit;
 };
diff --git a/Engine/Plugins/VirtualProduction/DMX/DMXPixelMapping/Source/DMXPixelMappingEditor/Private/Views/SDMXPixelMappingInputSourceView.cpp b/Engine/Plugins/VirtualProduction/DMX/DMXPixelMapping/Source/DMXPixelMappingEditor/Private/Views/SDMXPixelMappingInputSourceView.cpp
deleted file mode 100644
diff --git a/Engine/Plugins/VirtualProduction/DMX/DMXPixelMapping/Source/DMXPixelMappingEditor/Private/Views/SDMXPixelMappingInputSourceView.h b/Engine/Plugins/VirtualProduction/DMX/DMXPixelMapping/Source/DMXPixelMappingEditor/Private/Views/SDMXPixelMappingInputSourceView.h
deleted file mode 100644
diff --git a/Engine/Plugins/VirtualProduction/DMX/DMXPixelMapping/Source/DMXPixelMappingEditor/Private/Views/SDMXPixelMappingLayoutView.cpp b/Engine/Plugins/VirtualProduction/DMX/DMXPixelMapping/Source/DMXPixelMappingEditor/Private/Views/SDMXPixelMappingLayoutView.cpp
--- a/Engine/Plugins/VirtualProduction/DMX/DMXPixelMapping/Source/DMXPixelMappingEditor/Private/Views/SDMXPixelMappingLayoutView.cpp
+++ b/Engine/Plugins/VirtualProduction/DMX/DMXPixelMapping/Source/DMXPixelMappingEditor/Private/Views/SDMXPixelMappingLayoutView.cpp
@@ -2,17 +2,16 @@
 
 #include "Views/SDMXPixelMappingLayoutView.h"
 
-#include "DMXPixelMappingLayoutSettings.h"
 #include "Customizations/DMXPixelMappingHorizontalAlignmentCustomization.h"
 #include "Customizations/DMXPixelMappingLayoutViewModelDetails.h"
 #include "Customizations/DMXPixelMappingVerticalAlignmentCustomization.h"
-#include "Toolkits/DMXPixelMappingToolkit.h"
-#include "ViewModels/DMXPixelMappingLayoutViewModel.h"
-
 #include "DetailsViewArgs.h"
 #include "IDetailsView.h"
-#include "PropertyEditorModule.h"
 #include "Modules/ModuleManager.h"
+#include "PropertyEditorModule.h"
+#include "Settings/DMXPixelMappingEditorSettings.h"
+#include "Toolkits/DMXPixelMappingToolkit.h"
+#include "ViewModels/DMXPixelMappingLayoutViewModel.h"
 #include "Widgets/Layout/SScrollBox.h"
 
 
@@ -62,7 +61,6 @@ void SDMXPixelMappingLayoutView::Construct(const FArguments& InArgs, const TShar
 			[
 				CreateLayoutSettingsDetailsView()
 			]
-
 		];
 
 	Refresh();
@@ -111,11 +109,7 @@ void SDMXPixelMappingLayoutView::Refresh()
 
 TSharedRef<SWidget> SDMXPixelMappingLayoutView::CreateLayoutSettingsDetailsView()
 {
-	UDMXPixelMappingLayoutSettings* LayoutSettings = GetMutableDefault<UDMXPixelMappingLayoutSettings>();
-	if (!LayoutSettings)
-	{
-		return SNullWidget::NullWidget;
-	}
+	UDMXPixelMappingEditorSettings* LayoutSettings = GetMutableDefault<UDMXPixelMappingEditorSettings>();
 
 	FPropertyEditorModule& PropertyEditorModule = FModuleManager::GetModuleChecked<FPropertyEditorModule>("PropertyEditor");
 	FDetailsViewArgs DetailsViewArgs;
diff --git a/Engine/Plugins/VirtualProduction/DMX/DMXPixelMapping/Source/DMXPixelMappingEditor/Private/Widgets/SDMXPixelMappingHierarchyItem.cpp b/Engine/Plugins/VirtualProduction/DMX/DMXPixelMapping/Source/DMXPixelMappingEditor/Private/Widgets/SDMXPixelMappingHierarchyItem.cpp
--- a/Engine/Plugins/VirtualProduction/DMX/DMXPixelMapping/Source/DMXPixelMappingEditor/Private/Widgets/SDMXPixelMappingHierarchyItem.cpp
+++ b/Engine/Plugins/VirtualProduction/DMX/DMXPixelMapping/Source/DMXPixelMappingEditor/Private/Widgets/SDMXPixelMappingHierarchyItem.cpp
@@ -1,97 +1,186 @@
 // Copyright Epic Games, Inc. All Rights Reserved.
 
 #include "Widgets/SDMXPixelMappingHierarchyItem.h"
-#include "Components/DMXPixelMappingOutputComponent.h"
-#include "ViewModels/DMXPixelMappingHierarchyViewModel.h"
 
-#include "Widgets/Text/SInlineEditableTextBlock.h"
+#include "Components/DMXPixelMappingOutputComponent.h"
+#include "DMXEditorStyle.h"
+#include "DMXPixelMappingEditorUtils.h"
 #include "DragDrop/DMXPixelMappingDragDropOp.h"
-#include "Views/SDMXPixelMappingHierarchyView.h"
 #include "ScopedTransaction.h"
+#include "Toolkits/DMXPixelMappingToolkit.h"
+#include "Widgets/Text/SInlineEditableTextBlock.h"
 #include "Widgets/Views/SListView.h"
+#include "ViewModels/DMXPixelMappingHierarchyItem.h"
+#include "Views/SDMXPixelMappingHierarchyView.h"
 
-void SDMXPixelMappingHierarchyItem::Construct(const FArguments& InArgs, const TSharedRef<STableViewBase>& InOwnerTableView, TSharedPtr<FDMXPixelMappingHierarchyItemWidgetModel> InModel, TSharedPtr<SDMXPixelMappingHierarchyView> InHierarchyView)
-{
-	Model = InModel;
-	HierarchyView = InHierarchyView;
 
-	InModel->RenameEvent.BindSP(this, &SDMXPixelMappingHierarchyItem::OnRequestBeginRename);
+#define LOCTEXT_NAMESPACE "SDMXPixelMappingHierarchyItem"
 
-	STableRow<FDMXPixelMappingHierarchyItemWidgetModelPtr>::Construct(
-		STableRow<FDMXPixelMappingHierarchyItemWidgetModelPtr>::FArguments()
+void SDMXPixelMappingHierarchyItem::Construct(const FArguments& InArgs, const TSharedRef<STableViewBase>& InOwnerTableView, TWeakPtr<FDMXPixelMappingToolkit> InWeakToolkit, const TSharedRef<FDMXPixelMappingHierarchyItem>& InItem)
+{
+	WeakToolkit = InWeakToolkit;
+	Item = InItem;
+
+	SMultiColumnTableRow<FDMXPixelMappingHierarchyItemWidgetModelPtr>::Construct(
+		FSuperRowType::FArguments()
 		.Padding(0.0f)
 		.Style(&FAppStyle::Get().GetWidgetStyle<FTableRowStyle>("SimpleTableView.Row"))
 		.OnDragDetected(this, &SDMXPixelMappingHierarchyItem::OnDraggingWidget)
-		.OnDrop(this, &SDMXPixelMappingHierarchyItem::OnDropWidget)
-		.Content()
-		[
-			SNew(SHorizontalBox)
-			+ SHorizontalBox::Slot()
-			.FillWidth(1.0f)
-			.Padding(2, 0, 0, 0)
-			.VAlign(VAlign_Center)
-			[
-				SAssignNew(EditBox, SInlineEditableTextBlock)
-				.Text(this, &SDMXPixelMappingHierarchyItem::GetItemText)
-				.OnVerifyTextChanged(this, &SDMXPixelMappingHierarchyItem::OnVerifyNameTextChanged)
-				.OnTextCommitted(this, &SDMXPixelMappingHierarchyItem::OnNameTextCommited)
-				.IsEnabled(true)
-			]
-		],
+		.OnDrop(this, &SDMXPixelMappingHierarchyItem::OnDropWidget),
 		InOwnerTableView);
 }
 
-FText SDMXPixelMappingHierarchyItem::GetItemText() const
+void SDMXPixelMappingHierarchyItem::EnterRenameMode()
+{
+	EditableNameTextBox->EnterEditingMode();
+}
+
+TSharedRef<SWidget> SDMXPixelMappingHierarchyItem::GenerateWidgetForColumn(const FName& ColumnName)
+{
+	if (ColumnName == SDMXPixelMappingHierarchyView::FColumnIds::EditorColor)
+	{
+		return GenerateEditorColorWidget();
+	}
+	else if (ColumnName == SDMXPixelMappingHierarchyView::FColumnIds::ComponentName)
+	{			
+		// The name column gets the tree expansion arrow
+		return SNew(SBox)
+			[
+				SNew(SHorizontalBox)
+
+				+ SHorizontalBox::Slot()
+				.AutoWidth()
+				[
+					SNew(SExpanderArrow, SharedThis(this))
+				]
+
+				+ SHorizontalBox::Slot()
+				.FillWidth(1.f)
+				[
+					GenerateComponentNameWidget()
+				]
+			];
+	}
+	else if (ColumnName == SDMXPixelMappingHierarchyView::FColumnIds::FixtureID)
+	{
+		return GenerateFixtureIDWidget();
+	}
+	else if (ColumnName == SDMXPixelMappingHierarchyView::FColumnIds::Patch)
+	{
+		return GeneratePatchWidget();
+	}
+
+	return SNullWidget::NullWidget;
+}
+
+TSharedRef<SWidget> SDMXPixelMappingHierarchyItem::GenerateEditorColorWidget()
 {
-	if (FDMXPixelMappingHierarchyItemWidgetModelPtr ModelPtr = Model.Pin())
+	if (Item.IsValid() && Item->HasEditorColor())
 	{
-		return ModelPtr->GetText();
+		return
+			SNew(SBorder)
+			.HAlign(HAlign_Fill)
+			.Padding(5.f, 2.f)
+			.BorderImage(FAppStyle::GetBrush("NoBorder"))
+			[
+				SAssignNew(EditorColorImageWidget, SImage)
+				.Image(FDMXEditorStyle::Get().GetBrush("DMXEditor.WhiteRoundedPropertyBorder"))
+			.ColorAndOpacity_Lambda([this]() -> FSlateColor
+				{
+					return Item.IsValid() ? Item->GetEditorColor() : FLinearColor::Red;
+				})
+			];
 	}
 
-	return FText();
+	return SNullWidget::NullWidget;
 }
 
-bool SDMXPixelMappingHierarchyItem::OnVerifyNameTextChanged(const FText& InText, FText& OutErrorMessage)
+TSharedRef<SWidget> SDMXPixelMappingHierarchyItem::GenerateComponentNameWidget()
 {
-	return Model.Pin()->OnVerifyNameTextChanged(InText, OutErrorMessage);
+	return
+		SNew(SBorder)
+		.HAlign(HAlign_Fill)
+		.VAlign(VAlign_Center)
+		.Padding(4.f)
+		.BorderImage(FAppStyle::GetBrush("NoBorder"))
+		[
+			SAssignNew(EditableNameTextBox, SInlineEditableTextBlock)
+			.Font(FAppStyle::GetFontStyle(TEXT("PropertyWindow.NormalFont")))
+			.Text(Item.Get(), &FDMXPixelMappingHierarchyItem::GetComponentNameText)
+			.ColorAndOpacity(FLinearColor::White)
+		];
 }
 
-void SDMXPixelMappingHierarchyItem::OnNameTextCommited(const FText& InText, ETextCommit::Type CommitInfo)
+TSharedRef<SWidget> SDMXPixelMappingHierarchyItem::GenerateFixtureIDWidget()
 {
-	Model.Pin()->OnNameTextCommited(InText, CommitInfo);
+	return
+		SNew(SBorder)
+		.HAlign(HAlign_Fill)
+		.VAlign(VAlign_Center)
+		.Padding(4.f)
+		.BorderImage(FAppStyle::GetBrush("NoBorder"))
+		[
+			SNew(STextBlock)
+			.Font(FAppStyle::GetFontStyle(TEXT("PropertyWindow.NormalFont")))
+			.Text(Item.Get(), &FDMXPixelMappingHierarchyItem::GetFixtureIDText)
+			.ColorAndOpacity(FLinearColor::White)
+		];
+}
+
+TSharedRef<SWidget> SDMXPixelMappingHierarchyItem::GeneratePatchWidget()
+{
+	return
+		SNew(SBorder)
+		.HAlign(HAlign_Fill)
+		.VAlign(VAlign_Center)
+		.Padding(4.f)
+		.BorderImage(FAppStyle::GetBrush("NoBorder"))
+		[
+			SNew(STextBlock)
+			.Font(FAppStyle::GetFontStyle(TEXT("PropertyWindow.NormalFont")))
+			.Text(Item.Get(), &FDMXPixelMappingHierarchyItem::GetPatchText)
+			.ColorAndOpacity(FLinearColor::White)
+		];
 }
 
-void SDMXPixelMappingHierarchyItem::OnRequestBeginRename()
+void SDMXPixelMappingHierarchyItem::OnNameTextCommited(const FText& InText, ETextCommit::Type CommitInfo)
 {
-	if (EditBox.IsValid())
+	const TSharedPtr<FDMXPixelMappingToolkit> Toolkit = WeakToolkit.Pin();
+	const UDMXPixelMappingBaseComponent* Component = Item.IsValid() ? Item->GetComponent() : nullptr;
+	if (Toolkit.IsValid() && Component && !InText.IsEmpty())
 	{
-		EditBox->EnterEditingMode();
+		const FScopedTransaction RenameComponentTransaction(LOCTEXT("RenameComponentTransaction", "Rename Pixel Mapping Component"));
+
+		const FName& ComponentName = Component->GetFName();
+		Toolkit->RenameComponent(ComponentName, InText.ToString());
 	}
 }
 
 FReply SDMXPixelMappingHierarchyItem::OnDraggingWidget(const FGeometry& MyGeometry, const FPointerEvent& MouseEvent)
 {
-	check(Model.IsValid());
+	UDMXPixelMappingBaseComponent* Component = Item.IsValid() ? Item->GetComponent() : nullptr;
+	if (Component)
+	{
+		TArray<TWeakObjectPtr<UDMXPixelMappingBaseComponent>> DraggedComponents;
+		DraggedComponents.Add(Component);
 
-	UDMXPixelMappingBaseComponent* Component = Model.Pin()->GetReference().GetComponent();
+		return FReply::Handled().BeginDragDrop(FDMXPixelMappingDragDropOp::New(FVector2D::ZeroVector, DraggedComponents));
+	};
 
-	TArray<TWeakObjectPtr<UDMXPixelMappingBaseComponent>> DraggedComponents;
-	DraggedComponents.Add(Component);
-	return FReply::Handled().BeginDragDrop(FDMXPixelMappingDragDropOp::New(FVector2D::ZeroVector, DraggedComponents));
+	return FReply::Unhandled();
 }
 
 FReply SDMXPixelMappingHierarchyItem::OnDropWidget(const FDragDropEvent& InDragDropEvent)
 {
 	const FScopedTransaction Transaction(FText::FromString("MoveComponent"));
 	TSharedPtr<FDMXPixelMappingDragDropOp> ComponentDragDropOp = InDragDropEvent.GetOperationAs<FDMXPixelMappingDragDropOp>();
-	if (ComponentDragDropOp.IsValid()
-		&& HierarchyView.IsValid())
+	if (ComponentDragDropOp.IsValid())
 	{
 		for (const TWeakObjectPtr<UDMXPixelMappingBaseComponent>& WeakSource : ComponentDragDropOp->GetDraggedComponents())
 		{
 			if (UDMXPixelMappingBaseComponent* Source = WeakSource.Get())
 			{
-				UDMXPixelMappingBaseComponent* Destination = Model.Pin()->GetReference().GetComponent();
+				UDMXPixelMappingBaseComponent* Destination = Item->GetComponent();
 
 				UDMXPixelMappingBaseComponent* NewParent = nullptr;
 				if (Source->CanBeMovedTo(Destination))
@@ -134,9 +223,6 @@ FReply SDMXPixelMappingHierarchyItem::OnDropWidget(const FDragDropEvent& InDragD
 							}
 						}
 					}
-
-					HierarchyView->RequestComponentRedraw(NewParent);
-					HierarchyView->RequestRebuildTree();
 				}
 			}
 		}
@@ -144,3 +230,4 @@ FReply SDMXPixelMappingHierarchyItem::OnDropWidget(const FDragDropEvent& InDragD
 	return FReply::Unhandled();
 }
 
+#undef LOCTEXT_NAMESPACE
diff --git a/Engine/Plugins/VirtualProduction/DMX/DMXPixelMapping/Source/DMXPixelMappingEditor/Private/Widgets/SDMXPixelMappingHierarchyItem.h b/Engine/Plugins/VirtualProduction/DMX/DMXPixelMapping/Source/DMXPixelMappingEditor/Private/Widgets/SDMXPixelMappingHierarchyItem.h
--- a/Engine/Plugins/VirtualProduction/DMX/DMXPixelMapping/Source/DMXPixelMappingEditor/Private/Widgets/SDMXPixelMappingHierarchyItem.h
+++ b/Engine/Plugins/VirtualProduction/DMX/DMXPixelMapping/Source/DMXPixelMappingEditor/Private/Widgets/SDMXPixelMappingHierarchyItem.h
@@ -6,37 +6,63 @@
 #include "Widgets/Views/ITableRow.h"
 #include "Widgets/Views/STableRow.h"
 
-class STableViewBase;
+class FDMXPixelMappingHierarchyItem;
 class SInlineEditableTextBlock;
 class SDMXPixelMappingHierarchyView;
+class SImage;
+class SInlineEditableTextBox;
+class STableViewBase;
+
 
-class SDMXPixelMappingHierarchyItem
-	: public STableRow<TSharedPtr<FDMXPixelMappingHierarchyItemWidgetModel>>
+/** Displays a single item in a hierarchy view */
+class SDMXPixelMappingHierarchyItem final
+	: public SMultiColumnTableRow<TSharedPtr<FDMXPixelMappingHierarchyItem>>
 {
 public:
 	SLATE_BEGIN_ARGS(SDMXPixelMappingHierarchyItem) {}
 	SLATE_END_ARGS()
 
-	void Construct(const FArguments& InArgs, const TSharedRef< STableViewBase >& InOwnerTableView, TSharedPtr<FDMXPixelMappingHierarchyItemWidgetModel> Model, TSharedPtr<SDMXPixelMappingHierarchyView> InHierarchyView);
+	/** Constructs this widget */
+	void Construct(const FArguments& InArgs, const TSharedRef<STableViewBase>& InOwnerTableView, TWeakPtr<FDMXPixelMappingToolkit> InWeakToolkit, const TSharedRef<FDMXPixelMappingHierarchyItem>& InItem);
+
+	/** Starts editing the name of the component of this row */
+	void EnterRenameMode();
 
 private:
-	FText GetItemText() const;
+	//~ Begin SMultiColumnTableRow interface
+	virtual TSharedRef<SWidget> GenerateWidgetForColumn(const FName& ColumnName) override;
+	//~ End SMultiColumnTableRow interface
+
+	/** Generates a widget to display the editor color */
+	TSharedRef<SWidget> GenerateEditorColorWidget();
+
+	/** Generates a widget to display the name */
+	TSharedRef<SWidget> GenerateComponentNameWidget();
 
-	/** Called when text is being committed to check for validity */
-	bool OnVerifyNameTextChanged(const FText& InText, FText& OutErrorMessage);
+	/** Generates a widget to display the fixture ID */
+	TSharedRef<SWidget> GenerateFixtureIDWidget();
+
+	/** Generates a widget to display the patch */
+	TSharedRef<SWidget> GeneratePatchWidget();
 
 	/** Called when text is committed on the node */
 	void OnNameTextCommited(const FText& InText, ETextCommit::Type CommitInfo);
 
-	void OnRequestBeginRename();
-
+	/** Called when this row was dragged */
 	FReply OnDraggingWidget(const FGeometry& MyGeometry, const FPointerEvent& MouseEvent);
+
+	/** Called when this row was dropped */
 	FReply OnDropWidget(const FDragDropEvent& InDragDropEvent);
 
-private:
-	FDMXPixelMappingHierarchyWidgetModelWeakPtr Model;
-	TSharedPtr<SDMXPixelMappingHierarchyView> HierarchyView;
+	/** The editor color image widget */
+	TSharedPtr<SImage> EditorColorImageWidget;
+
+	/** Text box that displays and allows to edit the name */
+	TSharedPtr<SInlineEditableTextBlock> EditableNameTextBox;
+
+	/** Model of the component being displayed */
+	TSharedPtr<FDMXPixelMappingHierarchyItem> Item;
 
-	/** Edit box for the name. */
-	TSharedPtr<SInlineEditableTextBlock> EditBox;
+	/** Weak toolkit this widget belongs to */
+	TWeakPtr<FDMXPixelMappingToolkit> WeakToolkit;
 };
diff --git a/Engine/Plugins/VirtualProduction/DMX/DMXPixelMapping/Source/DMXPixelMappingEditor/Private/Widgets/SDMXPixelMappingOutputComponent.cpp b/Engine/Plugins/VirtualProduction/DMX/DMXPixelMapping/Source/DMXPixelMappingEditor/Private/Widgets/SDMXPixelMappingOutputComponent.cpp
--- a/Engine/Plugins/VirtualProduction/DMX/DMXPixelMapping/Source/DMXPixelMappingEditor/Private/Widgets/SDMXPixelMappingOutputComponent.cpp
+++ b/Engine/Plugins/VirtualProduction/DMX/DMXPixelMapping/Source/DMXPixelMappingEditor/Private/Widgets/SDMXPixelMappingOutputComponent.cpp
@@ -2,11 +2,11 @@
 
 #include "Widgets/SDMXPixelMappingOutputComponent.h"
 
-#include "DMXPixelMappingLayoutSettings.h"
 #include "DMXPixelMappingEditorStyle.h"
 #include "DMXPixelMappingTypes.h"
 #include "DMXPixelMappingUtils.h"
 #include "DMXRuntimeUtils.h"
+#include "Settings/DMXPixelMappingEditorSettings.h"
 #include "Toolkits/DMXPixelMappingToolkit.h"
 #include "ViewModels/DMXPixelMappingOutputComponentModel.h"
 #include "Widgets/Layout/SBorder.h"
@@ -82,7 +82,8 @@ void SDMXPixelMappingOutputComponent::Construct(const FArguments& InArgs, const
 	];
 
 	ForceRefresh();
-	UDMXPixelMappingLayoutSettings::GetOnLayoutSettingsChanged().AddSP(this, &SDMXPixelMappingOutputComponent::RequestRefresh);
+
+	UDMXPixelMappingEditorSettings::OnEditorSettingsChanged.AddSP(this, &SDMXPixelMappingOutputComponent::RequestRefresh);
 }
 
 bool SDMXPixelMappingOutputComponent::Equals(UDMXPixelMappingBaseComponent* Component) const
@@ -179,9 +180,6 @@ TSharedRef<SWidget> SDMXPixelMappingOutputComponent::CreateContent()
 			]
 		];
 
-	const UDMXPixelMappingLayoutSettings* LayoutSettings = GetDefault<UDMXPixelMappingLayoutSettings>();
-	check(LayoutSettings);
-
 	if (Model->ShouldDrawName())
 	{
 		if (Model->ShouldDrawNameAbove())
@@ -367,15 +365,15 @@ public:
 
 		RebuildAddressesContent();
 
-		UDMXPixelMappingLayoutSettings::GetOnLayoutSettingsChanged().AddSP(this, &SDMXPixelMappingScreenComponentCell::RebuildAddressesContent);
+		UDMXPixelMappingEditorSettings::OnEditorSettingsChanged.AddSP(this, &SDMXPixelMappingScreenComponentCell::RebuildAddressesContent);
 	}
 
 private:
 	/** Rebuilds addresses content */
 	void RebuildAddressesContent()
 	{
-		const UDMXPixelMappingLayoutSettings* LayoutSettings = GetDefault<UDMXPixelMappingLayoutSettings>();
-		if (!LayoutSettings || !LayoutSettings->bShowPatchInfo)
+		const FDMXPixelMappingDesignerSettings& DesignerSettings = GetDefault<UDMXPixelMappingEditorSettings>()->DesignerSettings;
+		if (!DesignerSettings.bShowPatchInfo)
 		{
 			return;
 		}
diff --git a/Engine/Plugins/VirtualProduction/DMX/DMXPixelMapping/Source/DMXPixelMappingEditor/Private/Widgets/SDMXPixelMappingTransformHandle.cpp b/Engine/Plugins/VirtualProduction/DMX/DMXPixelMapping/Source/DMXPixelMappingEditor/Private/Widgets/SDMXPixelMappingTransformHandle.cpp
--- a/Engine/Plugins/VirtualProduction/DMX/DMXPixelMapping/Source/DMXPixelMappingEditor/Private/Widgets/SDMXPixelMappingTransformHandle.cpp
+++ b/Engine/Plugins/VirtualProduction/DMX/DMXPixelMapping/Source/DMXPixelMappingEditor/Private/Widgets/SDMXPixelMappingTransformHandle.cpp
@@ -3,9 +3,9 @@
 #include "Widgets/SDMXPixelMappingTransformHandle.h"
 
 #include "Components/DMXPixelMappingMatrixCellComponent.h"
-#include "DMXPixelMappingLayoutSettings.h"
 #include "Framework/Application/SlateApplication.h"
 #include "ScopedTransaction.h"
+#include "Settings/DMXPixelMappingEditorSettings.h"
 #include "Toolkits/DMXPixelMappingToolkit.h"
 #include "Views/SDMXPixelMappingDesignerView.h"
 #include "Widgets/Images/SImage.h"
@@ -193,11 +193,9 @@ void SDMXPixelMappingTransformHandle::Resize(UDMXPixelMappingBaseComponent* Base
 
 		OutputComponent->SetSize(NewSize);
 
-		// Scale children if desired, no division by zero
-		const UDMXPixelMappingLayoutSettings* LayoutSettings = GetDefault<UDMXPixelMappingLayoutSettings>();
-		if (!LayoutSettings ||
-			!LayoutSettings->bScaleChildrenWithParent ||
-			NewSize == FVector2D::ZeroVector)
+		// Scale children only if desired, no division by zero
+		const FDMXPixelMappingDesignerSettings& DesignerSettings = GetDefault<UDMXPixelMappingEditorSettings>()->DesignerSettings;
+		if (!DesignerSettings.bScaleChildrenWithParent || NewSize == FVector2D::ZeroVector)
 		{
 			return;
 		}
diff --git a/Engine/Plugins/VirtualProduction/DMX/DMXPixelMapping/Source/DMXPixelMappingRuntime/Private/Components/DMXPixelMappingRendererComponent.cpp b/Engine/Plugins/VirtualProduction/DMX/DMXPixelMapping/Source/DMXPixelMappingRuntime/Private/Components/DMXPixelMappingRendererComponent.cpp
--- a/Engine/Plugins/VirtualProduction/DMX/DMXPixelMapping/Source/DMXPixelMappingRuntime/Private/Components/DMXPixelMappingRendererComponent.cpp
+++ b/Engine/Plugins/VirtualProduction/DMX/DMXPixelMapping/Source/DMXPixelMappingRuntime/Private/Components/DMXPixelMappingRendererComponent.cpp
@@ -17,6 +17,7 @@
 #include "Engine/TextureRenderTarget2D.h"
 #include "IDMXPixelMappingRenderer.h"
 #include "IDMXPixelMappingRendererModule.h"
+#include "Materials/MaterialInterface.h"
 #include "Modulators/DMXModulator.h"
 #include "RenderingThread.h"
 #include "TextureResource.h"
@@ -33,6 +34,8 @@ DECLARE_CYCLE_STAT(TEXT("PixelMapping SendDMX"), STAT_DMXPixelMappingSendDMX, ST
 DECLARE_CYCLE_STAT(TEXT("PixelMapping RenderInputTexture"), STAT_DMXPixelMappingRenderInputTexture, STATGROUP_DMX);
 
 
+#define LOCTEXT_NAMESPACE "DMXPixelMappingRendererComponent"
+
 UDMXPixelMappingRendererComponent::UDMXPixelMappingRendererComponent()
 {
 	SetSize(FVector2D(100.f, 100.f));
@@ -77,7 +80,7 @@ void UDMXPixelMappingRendererComponent::PostInitProperties()
 
 	if (!IsTemplate())
 	{
-		InvalidatePreprocessRenderer();
+		UpdatePreprocessRenderer();
 	}
 }
 
@@ -87,7 +90,7 @@ void UDMXPixelMappingRendererComponent::PostLoad()
 
 	if (!IsTemplate())
 	{
-		InvalidatePreprocessRenderer();
+		UpdatePreprocessRenderer();
 	}
 }
 
@@ -108,7 +111,7 @@ void UDMXPixelMappingRendererComponent::PostEditChangeChainProperty(FPropertyCha
 
 	if (PropertyChangedChainEvent.ChangeType != EPropertyChangeType::Interactive)
 	{
-		InvalidatePreprocessRenderer();
+		UpdatePreprocessRenderer();
 	}
 
 	PRAGMA_DISABLE_DEPRECATION_WARNINGS
@@ -125,7 +128,7 @@ void UDMXPixelMappingRendererComponent::PostEditChangeChainProperty(FPropertyCha
 }
 #endif // WITH_EDITOR
 
-void UDMXPixelMappingRendererComponent::InvalidatePreprocessRenderer()
+void UDMXPixelMappingRendererComponent::UpdatePreprocessRenderer()
 {
 	PRAGMA_DISABLE_DEPRECATION_WARNINGS
 	if (!PixelMappingRenderer_DEPRECATED.IsValid())
@@ -277,6 +280,30 @@ void UDMXPixelMappingRendererComponent::RenderAndSendDMX()
 	SendDMX();
 }
 
+FString UDMXPixelMappingRendererComponent::GetUserFriendlyName() const
+{
+	constexpr TCHAR NoSourceString[] = TEXT("None");
+	switch (RendererType)
+	{
+	case(EDMXPixelMappingRendererType::Texture):
+		return FString::Printf(TEXT("Pixel Mapping: %s"), InputTexture ? *InputTexture->GetName() : NoSourceString);
+		break;
+
+	case(EDMXPixelMappingRendererType::Material):
+		return FString::Printf(TEXT("Pixel Mapping: %s"), InputMaterial ? *InputMaterial->GetName() : NoSourceString);
+		break;
+
+	case(EDMXPixelMappingRendererType::UMG):
+		return FString::Printf(TEXT("Pixel Mapping: %s"), InputWidget ? *InputWidget->GetName() : NoSourceString);
+		break;
+
+	default:
+		checkf(0, TEXT("Invalid Renderer Type in DMXPixelMappingRendererComponent"));
+	}
+
+	return FString();
+}
+
 UTexture* UDMXPixelMappingRendererComponent::GetRenderedInputTexture() const
 {
 	return  PreprocessRenderer ? PreprocessRenderer->GetRenderedTexture() : nullptr;
@@ -683,3 +710,5 @@ UTextureRenderTarget2D* UDMXPixelMappingRendererComponent::CreateRenderTarget(co
 
 	return RenderTarget;
 }
+
+#undef LOCTEXT_NAMESPACE
diff --git a/Engine/Plugins/VirtualProduction/DMX/DMXPixelMapping/Source/DMXPixelMappingRuntime/Private/Components/DMXPixelMappingRootComponent.cpp b/Engine/Plugins/VirtualProduction/DMX/DMXPixelMapping/Source/DMXPixelMappingRuntime/Private/Components/DMXPixelMappingRootComponent.cpp
--- a/Engine/Plugins/VirtualProduction/DMX/DMXPixelMapping/Source/DMXPixelMappingRuntime/Private/Components/DMXPixelMappingRootComponent.cpp
+++ b/Engine/Plugins/VirtualProduction/DMX/DMXPixelMapping/Source/DMXPixelMappingRuntime/Private/Components/DMXPixelMappingRootComponent.cpp
@@ -108,3 +108,8 @@ void UDMXPixelMappingRootComponent::OnComponentRenamed(UDMXPixelMapping* PixelMa
 			CachedRendererComponentsByName.Add(RendererComponent->GetFName(), RendererComponent);
 		}, false);
 }
+
+FString UDMXPixelMappingRootComponent::GetUserFriendlyName() const
+{
+	return GetOuter()->GetName();
+}
diff --git a/Engine/Plugins/VirtualProduction/DMX/DMXPixelMapping/Source/DMXPixelMappingRuntime/Public/Components/DMXPixelMappingRendererComponent.h b/Engine/Plugins/VirtualProduction/DMX/DMXPixelMapping/Source/DMXPixelMappingRuntime/Public/Components/DMXPixelMappingRendererComponent.h
--- a/Engine/Plugins/VirtualProduction/DMX/DMXPixelMapping/Source/DMXPixelMappingRuntime/Public/Components/DMXPixelMappingRendererComponent.h
+++ b/Engine/Plugins/VirtualProduction/DMX/DMXPixelMapping/Source/DMXPixelMappingRuntime/Public/Components/DMXPixelMappingRendererComponent.h
@@ -53,13 +53,14 @@ public:
 	virtual void SendDMX() override;
 	virtual void Render() final;
 	virtual void RenderAndSendDMX() final;
+	virtual FString GetUserFriendlyName() const override;
 	//~ End UDMXPixelMappingBaseComponent implementation
 
 	/** Returns the pixel mapping preprocess renderer */
 	UDMXPixelMappingPreprocessRenderer* GetPreprocessRenderer() const { return PreprocessRenderer; }
 
-	/** Initializes the preprocess renderer from the current render type and input source */
-	void InvalidatePreprocessRenderer();
+	/** Updates the preprocess renderer from the current render type and source texture */
+	void UpdatePreprocessRenderer();
 
 	/**
 	 * Invalidates the pixel map, effectively causing a rebuild of the mapping on the next Render call.
diff --git a/Engine/Plugins/VirtualProduction/DMX/DMXPixelMapping/Source/DMXPixelMappingRuntime/Public/Components/DMXPixelMappingRootComponent.h b/Engine/Plugins/VirtualProduction/DMX/DMXPixelMapping/Source/DMXPixelMappingRuntime/Public/Components/DMXPixelMappingRootComponent.h
--- a/Engine/Plugins/VirtualProduction/DMX/DMXPixelMapping/Source/DMXPixelMappingRuntime/Public/Components/DMXPixelMappingRootComponent.h
+++ b/Engine/Plugins/VirtualProduction/DMX/DMXPixelMapping/Source/DMXPixelMappingRuntime/Public/Components/DMXPixelMappingRootComponent.h
@@ -33,6 +33,7 @@ public:
 	virtual void SendDMX() override;
 	virtual void Render() override;
 	virtual void RenderAndSendDMX() override;
+	virtual FString GetUserFriendlyName() const override;
 	//~ End UDMXPixelMappingBaseComponent implementation
 
 	/** Check if a Component can be moved under another one (used for copy/move/duplicate) */
