DMX - Fix Pixel Mapping hierarchy colapses when search is cleared, expansion state gets stored

diff --git a/Engine/Plugins/VirtualProduction/DMX/DMXPixelMapping/Source/DMXPixelMappingEditor/Private/ViewModels/DMXPixelMappingHierarchyItem.cpp b/Engine/Plugins/VirtualProduction/DMX/DMXPixelMapping/Source/DMXPixelMappingEditor/Private/ViewModels/DMXPixelMappingHierarchyItem.cpp
--- a/Engine/Plugins/VirtualProduction/DMX/DMXPixelMapping/Source/DMXPixelMappingEditor/Private/ViewModels/DMXPixelMappingHierarchyItem.cpp
+++ b/Engine/Plugins/VirtualProduction/DMX/DMXPixelMapping/Source/DMXPixelMappingEditor/Private/ViewModels/DMXPixelMappingHierarchyItem.cpp
@@ -90,6 +90,26 @@ int64 FDMXPixelMappingHierarchyItem::GetAbsoluteChannel() const
 	return -1;
 }
 
+bool FDMXPixelMappingHierarchyItem::IsExpanded() const
+{
+	if (UDMXPixelMappingBaseComponent* BaseComponent = Cast<UDMXPixelMappingBaseComponent>(WeakComponent.Get()))
+	{
+		return BaseComponent->bExpanded;
+	}
+
+	return false;
+}
+
+void FDMXPixelMappingHierarchyItem::SetIsExpanded(bool bExpanded)
+{
+	UDMXPixelMappingBaseComponent* BaseComponent = Cast<UDMXPixelMappingBaseComponent>(WeakComponent.Get());
+	if (BaseComponent && BaseComponent->bExpanded != bExpanded)
+	{
+		BaseComponent->Modify();
+		BaseComponent->bExpanded = bExpanded;
+	}
+}
+
 void FDMXPixelMappingHierarchyItem::ReverseChildren()
 {
 	Algo::Reverse(Children);
@@ -167,4 +187,4 @@ void FDMXPixelMappingHierarchyItem::UpdateFixtureID()
 			}
 		}
 	}
-}
\ No newline at end of file
+}
diff --git a/Engine/Plugins/VirtualProduction/DMX/DMXPixelMapping/Source/DMXPixelMappingEditor/Private/ViewModels/DMXPixelMappingHierarchyItem.h b/Engine/Plugins/VirtualProduction/DMX/DMXPixelMapping/Source/DMXPixelMappingEditor/Private/ViewModels/DMXPixelMappingHierarchyItem.h
--- a/Engine/Plugins/VirtualProduction/DMX/DMXPixelMapping/Source/DMXPixelMappingEditor/Private/ViewModels/DMXPixelMappingHierarchyItem.h
+++ b/Engine/Plugins/VirtualProduction/DMX/DMXPixelMapping/Source/DMXPixelMappingEditor/Private/ViewModels/DMXPixelMappingHierarchyItem.h
@@ -44,6 +44,12 @@ public:
 	/** Gets the absolute channel of the patch of the component. Returns -1 if this item doesn't have a patch */
 	int64 GetAbsoluteChannel() const;
 
+	/** Returns true if this item is expanded */
+	bool IsExpanded() const;
+
+	/** Sets if the item is expanded */
+	void SetIsExpanded(bool bExpanded);
+
 	/** Sorts the children of this item by Projection. */
 	template <typename ProjectionType>
 	void StableSortChildren(ProjectionType Projection)
diff --git a/Engine/Plugins/VirtualProduction/DMX/DMXPixelMapping/Source/DMXPixelMappingEditor/Private/Views/SDMXPixelMappingHierarchyView.cpp b/Engine/Plugins/VirtualProduction/DMX/DMXPixelMapping/Source/DMXPixelMappingEditor/Private/Views/SDMXPixelMappingHierarchyView.cpp
--- a/Engine/Plugins/VirtualProduction/DMX/DMXPixelMapping/Source/DMXPixelMappingEditor/Private/Views/SDMXPixelMappingHierarchyView.cpp
+++ b/Engine/Plugins/VirtualProduction/DMX/DMXPixelMapping/Source/DMXPixelMappingEditor/Private/Views/SDMXPixelMappingHierarchyView.cpp
@@ -68,80 +68,16 @@ namespace UE::DMX::PixelMappingEditor::SDMXPixelMappingHierarchyView::Private
 		TArray<UDMXPixelMappingBaseComponent*> DMXPixelMappingBaseComponents;
 	};
 
-
-	/** Helper to restore expansion state after refreshing */
-	struct FTreeExpansionSnapshot
+	/** Helper to adopt the expansion of items in theh hierarchy treeview */
+	void AdoptItemExpansionFromTree(const TSharedRef<STreeView<FDMXPixelMappingHierarchyItemWidgetModelPtr>>& TreeView, const TSharedPtr<FDMXPixelMappingHierarchyItem>& Parent)
 	{
-		static FTreeExpansionSnapshot TakeSnapshot(const FDMXPixelMappingHierarchyItemWidgetModelArr& InRootWidgets, const TSharedPtr<STreeView<FDMXPixelMappingHierarchyItemWidgetModelPtr>>& InTreeView)
-		{
-			FTreeExpansionSnapshot Result;
-			for (FDMXPixelMappingHierarchyItemWidgetModelPtr Model : InRootWidgets)
-			{
-				Result.RecursiveTakeSnapshot(Model, InTreeView);
-			}
-			return Result;
-		}
-
-		void RestoreExpandedAndExpandNewModels(const FDMXPixelMappingHierarchyItemWidgetModelArr& InRootWidgets, const TSharedPtr<STreeView<FDMXPixelMappingHierarchyItemWidgetModelPtr>>& InTreeView)
+		TreeView->SetItemExpansion(Parent, Parent->IsExpanded());
+		for (const TSharedPtr<FDMXPixelMappingHierarchyItem>& Child : Parent->GetChildren())
 		{
-			for (FDMXPixelMappingHierarchyItemWidgetModelPtr Model : InRootWidgets)
-			{
-				RecursiveRestoreSnapshot(Model, InTreeView);
-			}
+			TreeView->SetItemExpansion(Child, Child->IsExpanded());
+			AdoptItemExpansionFromTree(TreeView, Child);
 		}
-
-	private:
-
-		FTreeExpansionSnapshot()
-		{}
-
-		void RecursiveTakeSnapshot(FDMXPixelMappingHierarchyItemWidgetModelPtr Model, const TSharedPtr<STreeView<FDMXPixelMappingHierarchyItemWidgetModelPtr>>& TreeView)
-		{
-			UDMXPixelMappingBaseComponent* Component = Model->GetComponent();
-			if (IsValid(Component))
-			{
-				ComponentExpansionStates.Add(Component) = TreeView->IsItemExpanded(Model);
-
-				for (FDMXPixelMappingHierarchyItemWidgetModelPtr& ChildModel : Model->GetChildren())
-				{
-					RecursiveTakeSnapshot(ChildModel, TreeView);
-				}
-			}
-		}
-
-		void RecursiveRestoreSnapshot(FDMXPixelMappingHierarchyItemWidgetModelPtr Model, const TSharedPtr<STreeView<FDMXPixelMappingHierarchyItemWidgetModelPtr>>& TreeView)
-		{
-			UDMXPixelMappingBaseComponent* Component = Model->GetComponent();
-			if (IsValid(Component))
-			{
-				bool* pPreviousExpansionState = ComponentExpansionStates.Find(Component);
-				if (pPreviousExpansionState == nullptr)
-				{
-					// Initially collapse matrix components
-					if (Cast<UDMXPixelMappingMatrixComponent>(Component))
-					{
-						TreeView->SetItemExpansion(Model, false);
-					}
-					else
-					{
-						TreeView->SetItemExpansion(Model, true);
-					}
-				}
-				else
-				{
-					TreeView->SetItemExpansion(Model, *pPreviousExpansionState);
-				}
-
-				for (FDMXPixelMappingHierarchyItemWidgetModelPtr& ChildModel : Model->GetChildren())
-				{
-					RecursiveRestoreSnapshot(ChildModel, TreeView);
-				}
-			}
-		}
-
-		TMap<UDMXPixelMappingBaseComponent*, bool> ComponentExpansionStates;
-	};
-
+	}
 }
 
 const FName SDMXPixelMappingHierarchyView::FColumnIds::EditorColor = "EditorColor";
@@ -191,27 +127,30 @@ void SDMXPixelMappingHierarchyView::RequestRefresh()
 
 void SDMXPixelMappingHierarchyView::ForceRefresh()
 {
-	RequestRefreshTimerHandle.Invalidate();
+	if (!ensureMsgf(HierarchyTreeView.IsValid(), TEXT("Cannot refresh the hierarchy tree view before BuildChildSlotAndRefresh() was called.")))
+	{
+		return;
+	}
 
-	using namespace UE::DMX::PixelMappingEditor::SDMXPixelMappingHierarchyView::Private;
-	FTreeExpansionSnapshot ExpansionSnapshot = FTreeExpansionSnapshot::TakeSnapshot(AllRootItems, HierarchyTreeView);
+	RequestRefreshTimerHandle.Invalidate();
 
 	// Create the root, and let it construct new tree items
 	AllRootItems.Reset();
 	AllRootItems.Add(FDMXPixelMappingHierarchyItem::CreateNew(WeakToolkit.Pin()));
 
-	// Refresh and filter the tree view
-	FilterHandler->RefreshAndFilterTree();
+	// Adopt expansion of new items from the old tree
+	using namespace UE::DMX::PixelMappingEditor::SDMXPixelMappingHierarchyView::Private;
+	for (const TSharedPtr<FDMXPixelMappingHierarchyItem>& RootItem : AllRootItems)
+	{
+		AdoptItemExpansionFromTree(HierarchyTreeView.ToSharedRef(), RootItem);
+	}
 
-	ExpansionSnapshot.RestoreExpandedAndExpandNewModels(AllRootItems, HierarchyTreeView);
+	FilterHandler->RefreshAndFilterTree();
 	AdoptSelectionFromToolkit();
 }
 
 void SDMXPixelMappingHierarchyView::BuildChildSlotAndRefresh()
 {
-	using namespace UE::DMX::PixelMappingEditor::SDMXPixelMappingHierarchyView::Private;
-	FTreeExpansionSnapshot ExpansionSnapshot = FTreeExpansionSnapshot::TakeSnapshot(AllRootItems, HierarchyTreeView);
-
 	ChildSlot
 		[
 			SNew(SVerticalBox)
@@ -257,7 +196,8 @@ void SDMXPixelMappingHierarchyView::BuildChildSlotAndRefresh()
 				.HeaderRow(GenerateHeaderRow())
 				.OnGetChildren(FilterHandler.ToSharedRef(), &TreeFilterHandler<TSharedPtr<FDMXPixelMappingHierarchyItem>>::OnGetFilteredChildren)
 				.OnGenerateRow(this, &SDMXPixelMappingHierarchyView::OnGenerateRow)
-				.OnSelectionChanged(this, &SDMXPixelMappingHierarchyView::OnSelectionChanged)
+				.OnSelectionChanged(this, &SDMXPixelMappingHierarchyView::OnHierarchySelectionChanged)
+				.OnExpansionChanged(this, &SDMXPixelMappingHierarchyView::OnHierarchyExpansionChanged)
 				.OnContextMenuOpening(this, &SDMXPixelMappingHierarchyView::OnContextMenuOpening)
 				.TreeItemsSource(&FilteredRootItems)
 				.OnItemToString_Debug_Lambda([this](FDMXPixelMappingHierarchyItemWidgetModelPtr Item) { return Item->GetComponentNameText().ToString(); })
@@ -266,15 +206,7 @@ void SDMXPixelMappingHierarchyView::BuildChildSlotAndRefresh()
 
 	FilterHandler->SetTreeView(HierarchyTreeView.Get());
 
-	// Create the root, and let it construct new tree items
-	AllRootItems.Reset();
-	AllRootItems.Add(FDMXPixelMappingHierarchyItem::CreateNew(WeakToolkit.Pin()));
-
-	// Refresh and filter the tree view
-	FilterHandler->RefreshAndFilterTree();
-
-	ExpansionSnapshot.RestoreExpandedAndExpandNewModels(AllRootItems, HierarchyTreeView);
-	AdoptSelectionFromToolkit();
+	ForceRefresh();
 }
 
 TSharedRef<SHeaderRow> SDMXPixelMappingHierarchyView::GenerateHeaderRow()
@@ -472,7 +404,7 @@ TSharedPtr<SWidget> SDMXPixelMappingHierarchyView::OnContextMenuOpening()
 	return MenuBuilder.MakeWidget();
 }
 
-void SDMXPixelMappingHierarchyView::OnSelectionChanged(FDMXPixelMappingHierarchyItemWidgetModelPtr SelectedItem, ESelectInfo::Type SelectInfo)
+void SDMXPixelMappingHierarchyView::OnHierarchySelectionChanged(FDMXPixelMappingHierarchyItemWidgetModelPtr SelectedItem, ESelectInfo::Type SelectInfo)
 {
 	if (SelectInfo == ESelectInfo::Direct)
 	{
@@ -494,6 +426,11 @@ void SDMXPixelMappingHierarchyView::OnSelectionChanged(FDMXPixelMappingHierarchy
 	}
 }
 
+void SDMXPixelMappingHierarchyView::OnHierarchyExpansionChanged(TSharedPtr<FDMXPixelMappingHierarchyItem> Item, bool bExpanded)
+{
+	Item->SetIsExpanded(bExpanded);
+}
+
 void SDMXPixelMappingHierarchyView::ToggleColumnVisility(FName ColumnId)
 {
 	UDMXPixelMappingEditorSettings* EditorSettings = GetMutableDefault<UDMXPixelMappingEditorSettings>();
@@ -607,16 +544,9 @@ bool SDMXPixelMappingHierarchyView::RecursiveAdoptSelectionFromToolkit(const TSh
 	}
 
 	const TSet<FDMXPixelMappingComponentReference> SelectedComponents = Toolkit->GetSelectedComponents();
-	if (bContainsSelection)
-	{
-		HierarchyTreeView->SetItemExpansion(Item, true);
-	}
-
 	if (SelectedComponents.Contains(FDMXPixelMappingComponentReference(Toolkit, Item->GetComponent())))
 	{
-		HierarchyTreeView->SetItemSelection(Item, true, ESelectInfo::Direct);
 		HierarchyTreeView->RequestScrollIntoView(Item);
-
 		return true;
 	}
 
@@ -658,14 +588,9 @@ void SDMXPixelMappingHierarchyView::DeleteSelectedComponents()
 
 void SDMXPixelMappingHierarchyView::SetFilterText(const FText& Text)
 {
-	using namespace UE::DMX::PixelMappingEditor::SDMXPixelMappingHierarchyView::Private;
-	FTreeExpansionSnapshot ExpansionSnapshot = FTreeExpansionSnapshot::TakeSnapshot(AllRootItems, HierarchyTreeView);
-
 	FilterHandler->SetIsEnabled(!Text.IsEmpty());
 	SearchFilter->SetRawFilterText(Text);
 	FilterHandler->RefreshAndFilterTree();
-
-	ExpansionSnapshot.RestoreExpandedAndExpandNewModels(AllRootItems, HierarchyTreeView);
 }
 
 void SDMXPixelMappingHierarchyView::GetWidgetFilterStrings(FDMXPixelMappingHierarchyItemWidgetModelPtr InModel, TArray<FString>& OutStrings) const
diff --git a/Engine/Plugins/VirtualProduction/DMX/DMXPixelMapping/Source/DMXPixelMappingEditor/Private/Views/SDMXPixelMappingHierarchyView.h b/Engine/Plugins/VirtualProduction/DMX/DMXPixelMapping/Source/DMXPixelMappingEditor/Private/Views/SDMXPixelMappingHierarchyView.h
--- a/Engine/Plugins/VirtualProduction/DMX/DMXPixelMapping/Source/DMXPixelMappingEditor/Private/Views/SDMXPixelMappingHierarchyView.h
+++ b/Engine/Plugins/VirtualProduction/DMX/DMXPixelMapping/Source/DMXPixelMappingEditor/Private/Views/SDMXPixelMappingHierarchyView.h
@@ -74,16 +74,19 @@ private:
 	//~ End FEditorUndoClient interface
 
 	/** Called to get child items */
-	void OnGetChildItems(FDMXPixelMappingHierarchyItemWidgetModelPtr InParent, FDMXPixelMappingHierarchyItemWidgetModelArr& OutChildren);
+	void OnGetChildItems(TSharedPtr<FDMXPixelMappingHierarchyItem> InParent, FDMXPixelMappingHierarchyItemWidgetModelArr& OutChildren);
 
 	/** Called when a row is generated */
-	TSharedRef<ITableRow> OnGenerateRow(FDMXPixelMappingHierarchyItemWidgetModelPtr Item, const TSharedRef<STableViewBase>& OwnerTable);
+	TSharedRef<ITableRow> OnGenerateRow(TSharedPtr<FDMXPixelMappingHierarchyItem> Item, const TSharedRef<STableViewBase>& OwnerTable);
 
 	/** Called when the context menu is opening */
 	TSharedPtr<SWidget> OnContextMenuOpening();
 
-	/** Called when selection changed */
-	void OnSelectionChanged(FDMXPixelMappingHierarchyItemWidgetModelPtr SelectedItem, ESelectInfo::Type SelectInfo);
+	/** Called when selection changed in the hierarchy */
+	void OnHierarchySelectionChanged(TSharedPtr<FDMXPixelMappingHierarchyItem> Item, ESelectInfo::Type SelectInfo);
+
+	/** Called when selection changed in the hierarchy */
+	void OnHierarchyExpansionChanged(TSharedPtr<FDMXPixelMappingHierarchyItem> Item, bool bExpanded);
 
 	/** Toggles the visibility of a column */
 	void ToggleColumnVisility(FName ColumnId);
@@ -125,7 +128,7 @@ private:
 	void SetFilterText(const FText& Text);
 
 	/**  Gets an array of strings used for filtering/searching the specified widget. */
-	void GetWidgetFilterStrings(FDMXPixelMappingHierarchyItemWidgetModelPtr InModel, TArray<FString>& OutStrings) const;
+	void GetWidgetFilterStrings(TSharedPtr<FDMXPixelMappingHierarchyItem> InModel, TArray<FString>& OutStrings) const;
 
 	/** Flag to ignore selections while the hierarchy view is updating the selection. */
 	bool bIsUpdatingSelection = false;
@@ -143,13 +146,13 @@ private:
 	TSharedPtr<SSearchBox> SearchBox;
 
 	/** The hierarchy tree view widget */
-	TSharedPtr<STreeView<FDMXPixelMappingHierarchyItemWidgetModelPtr>> HierarchyTreeView;
+	TSharedPtr<STreeView<TSharedPtr<FDMXPixelMappingHierarchyItem>>> HierarchyTreeView;
 
 	/** Handles filtering the hierarchy based on an IFilter. */
-	TSharedPtr<TreeFilterHandler<FDMXPixelMappingHierarchyItemWidgetModelPtr>> FilterHandler;
+	TSharedPtr<TreeFilterHandler<TSharedPtr<FDMXPixelMappingHierarchyItem>>> FilterHandler;
 
 	/** Text filter for the filter handler, set from the search bbox */
-	TSharedPtr<TTextFilter<FDMXPixelMappingHierarchyItemWidgetModelPtr>> SearchFilter;
+	TSharedPtr<TTextFilter<TSharedPtr<FDMXPixelMappingHierarchyItem>>> SearchFilter;
 
 	/** Timer handle for the request refresh timer */
 	FTimerHandle RequestRefreshTimerHandle;
diff --git a/Engine/Plugins/VirtualProduction/DMX/DMXPixelMapping/Source/DMXPixelMappingRuntime/Private/Components/DMXPixelMappingMatrixCellComponent.cpp b/Engine/Plugins/VirtualProduction/DMX/DMXPixelMapping/Source/DMXPixelMappingRuntime/Private/Components/DMXPixelMappingMatrixCellComponent.cpp
--- a/Engine/Plugins/VirtualProduction/DMX/DMXPixelMapping/Source/DMXPixelMappingRuntime/Private/Components/DMXPixelMappingMatrixCellComponent.cpp
+++ b/Engine/Plugins/VirtualProduction/DMX/DMXPixelMapping/Source/DMXPixelMappingRuntime/Private/Components/DMXPixelMappingMatrixCellComponent.cpp
@@ -29,7 +29,7 @@ UDMXPixelMappingMatrixCellComponent::UDMXPixelMappingMatrixCellComponent()
 {
 #if WITH_EDITORONLY_DATA
 	bLockInDesigner = true;
-	EditorColor = FLinearColor::White.CopyWithNewOpacity(.4f);
+	EditorColor = FLinearColor::White.CopyWithNewOpacity(.25f);
 #endif // WITH_EDITORONLY_DATA
 }
 
diff --git a/Engine/Plugins/VirtualProduction/DMX/DMXPixelMapping/Source/DMXPixelMappingRuntime/Private/Components/DMXPixelMappingMatrixComponent.cpp b/Engine/Plugins/VirtualProduction/DMX/DMXPixelMapping/Source/DMXPixelMappingRuntime/Private/Components/DMXPixelMappingMatrixComponent.cpp
--- a/Engine/Plugins/VirtualProduction/DMX/DMXPixelMapping/Source/DMXPixelMappingRuntime/Private/Components/DMXPixelMappingMatrixComponent.cpp
+++ b/Engine/Plugins/VirtualProduction/DMX/DMXPixelMapping/Source/DMXPixelMappingRuntime/Private/Components/DMXPixelMappingMatrixComponent.cpp
@@ -38,6 +38,8 @@ UDMXPixelMappingMatrixComponent::UDMXPixelMappingMatrixComponent()
 	SetSize(FVector2D(32.f, 32.f));
 
 #if WITH_EDITORONLY_DATA
+	bExpanded = false;
+
 	// Even tho deprecated, default values on deprecated properties need be set so they don't load their type's default value.
 	ColorMode_DEPRECATED = EDMXColorMode::CM_RGB;
 	AttributeR_DEPRECATED.SetFromName("Red");
diff --git a/Engine/Plugins/VirtualProduction/DMX/DMXPixelMapping/Source/DMXPixelMappingRuntime/Public/Components/DMXPixelMappingBaseComponent.h b/Engine/Plugins/VirtualProduction/DMX/DMXPixelMapping/Source/DMXPixelMappingRuntime/Public/Components/DMXPixelMappingBaseComponent.h
--- a/Engine/Plugins/VirtualProduction/DMX/DMXPixelMapping/Source/DMXPixelMappingRuntime/Public/Components/DMXPixelMappingBaseComponent.h
+++ b/Engine/Plugins/VirtualProduction/DMX/DMXPixelMapping/Source/DMXPixelMappingRuntime/Public/Components/DMXPixelMappingBaseComponent.h
@@ -188,6 +188,11 @@ public:
 	UFUNCTION(BlueprintCallable, Category = "DMX|PixelMapping")
 	virtual void RenderAndSendDMX() {};
 
+#if WITH_EDITORONLY_DATA
+	UPROPERTY()
+	bool bExpanded = true;
+#endif
+
 public:
 	/*----------------------------------------------------------
 		Public static functions
