DMX - Pixel mapping components can be dragged out of the fixture patch list and hierarchy again

diff --git a/Engine/Plugins/VirtualProduction/DMX/DMXEngine/Source/DMXEditor/Private/Widgets/FixturePatch/SDMXReadOnlyFixturePatchList.cpp b/Engine/Plugins/VirtualProduction/DMX/DMXEngine/Source/DMXEditor/Private/Widgets/FixturePatch/SDMXReadOnlyFixturePatchList.cpp
--- a/Engine/Plugins/VirtualProduction/DMX/DMXEngine/Source/DMXEditor/Private/Widgets/FixturePatch/SDMXReadOnlyFixturePatchList.cpp
+++ b/Engine/Plugins/VirtualProduction/DMX/DMXEngine/Source/DMXEditor/Private/Widgets/FixturePatch/SDMXReadOnlyFixturePatchList.cpp
@@ -73,15 +73,10 @@ void SDMXReadOnlyFixturePatchList::Construct(const FArguments& InArgs)
 
 	IsRowEnabledDelegate = InArgs._IsRowEnabled;
 	IsRowVisibleDelegate = InArgs._IsRowVisibile;
+	OnRowDragDetectedDelegate = InArgs._OnRowDragDetected;
 
 	UpdateListItems();
 
-	// Handle Entity changes
-	UDMXLibrary::GetOnEntitiesAdded().AddSP(this, &SDMXReadOnlyFixturePatchList::OnEntityAddedOrRemoved);
-	UDMXLibrary::GetOnEntitiesRemoved().AddSP(this, &SDMXReadOnlyFixturePatchList::OnEntityAddedOrRemoved);
-	UDMXEntityFixturePatch::GetOnFixturePatchChanged().AddSP(this, &SDMXReadOnlyFixturePatchList::OnFixturePatchChanged);
-	UDMXEntityFixtureType::GetOnFixtureTypeChanged().AddSP(this, &SDMXReadOnlyFixturePatchList::OnFixtureTypeChanged);
-
 	ChildSlot
 		[
 			SNew(SVerticalBox)
@@ -134,6 +129,13 @@ void SDMXReadOnlyFixturePatchList::Construct(const FArguments& InArgs)
 			]
 		];
 
+
+	// Handle Entity changes
+	UDMXLibrary::GetOnEntitiesAdded().AddSP(this, &SDMXReadOnlyFixturePatchList::OnEntityAddedOrRemoved);
+	UDMXLibrary::GetOnEntitiesRemoved().AddSP(this, &SDMXReadOnlyFixturePatchList::OnEntityAddedOrRemoved);
+	UDMXEntityFixturePatch::GetOnFixturePatchChanged().AddSP(this, &SDMXReadOnlyFixturePatchList::OnFixturePatchChanged);
+	UDMXEntityFixtureType::GetOnFixtureTypeChanged().AddSP(this, &SDMXReadOnlyFixturePatchList::OnFixtureTypeChanged);
+
 	InitializeByListDescriptor(InArgs._ListDescriptor);
 }
 
@@ -353,7 +355,8 @@ TSharedRef<ITableRow> SDMXReadOnlyFixturePatchList::OnGenerateRow(TSharedPtr<FDM
 	const TSharedRef<SDMXReadOnlyFixturePatchListRow> NewRow =
 		SNew(SDMXReadOnlyFixturePatchListRow, OwnerTable, InItem.ToSharedRef())
 		.IsEnabled(this, &SDMXReadOnlyFixturePatchList::IsRowEnabled, InItem)
-		.Visibility(TAttribute<EVisibility>::CreateSP(this, &SDMXReadOnlyFixturePatchList::GetRowVisibility, InItem));
+		.Visibility(TAttribute<EVisibility>::CreateSP(this, &SDMXReadOnlyFixturePatchList::GetRowVisibility, InItem))
+		.OnRowDragDetected(OnRowDragDetectedDelegate);
 
 	return NewRow;
 }
diff --git a/Engine/Plugins/VirtualProduction/DMX/DMXEngine/Source/DMXEditor/Private/Widgets/FixturePatch/SDMXReadOnlyFixturePatchListRow.cpp b/Engine/Plugins/VirtualProduction/DMX/DMXEngine/Source/DMXEditor/Private/Widgets/FixturePatch/SDMXReadOnlyFixturePatchListRow.cpp
--- a/Engine/Plugins/VirtualProduction/DMX/DMXEngine/Source/DMXEditor/Private/Widgets/FixturePatch/SDMXReadOnlyFixturePatchListRow.cpp
+++ b/Engine/Plugins/VirtualProduction/DMX/DMXEngine/Source/DMXEditor/Private/Widgets/FixturePatch/SDMXReadOnlyFixturePatchListRow.cpp
@@ -21,7 +21,7 @@ void SDMXReadOnlyFixturePatchListRow::Construct(const FArguments& InArgs, const
 
 	SMultiColumnTableRow<TSharedPtr<FDMXEntityFixturePatchRef>>::Construct(
 		FSuperRowType::FArguments()
-		.OnDragDetected(InArgs._OnRowDragged)
+		.OnDragDetected(InArgs._OnRowDragDetected)
 		.Style(&FDMXEditorStyle::Get().GetWidgetStyle<FTableRowStyle>("FixturePatchList.Row")),
 		InOwnerTable);
 }
diff --git a/Engine/Plugins/VirtualProduction/DMX/DMXEngine/Source/DMXEditor/Public/Widgets/SDMXReadOnlyFixturePatchList.h b/Engine/Plugins/VirtualProduction/DMX/DMXEngine/Source/DMXEditor/Public/Widgets/SDMXReadOnlyFixturePatchList.h
--- a/Engine/Plugins/VirtualProduction/DMX/DMXEngine/Source/DMXEditor/Public/Widgets/SDMXReadOnlyFixturePatchList.h
+++ b/Engine/Plugins/VirtualProduction/DMX/DMXEngine/Source/DMXEditor/Public/Widgets/SDMXReadOnlyFixturePatchList.h
@@ -112,7 +112,7 @@ public:
 		SLATE_EVENT(FDMXFixturePatchListRowRetValDelegate, IsRowVisibile)
 
 		/** Delegate executed when a row was dragged */
-		SLATE_EVENT(FOnDragDetected, OnRowDragged)
+		SLATE_EVENT(FOnDragDetected, OnRowDragDetected)
 
 	SLATE_END_ARGS()
 
@@ -168,9 +168,6 @@ protected:
 	/** Called to get wheter a row of the list is visible or not */
 	virtual EVisibility GetRowVisibility(const TSharedPtr<FDMXEntityFixturePatchRef> InFixturePatchRef) const;
 
-	/** Called when a row was dragged */
-	FReply OnRowDragged(const FGeometry& MyGeometry, const FPointerEvent& MouseEvent);
-
 	/** Called when text from searchbox is changed */
 	void OnSearchTextChanged(const FText& SearchText);
 
@@ -244,5 +241,5 @@ protected:
 	TArray<FDMXEntityFixturePatchRef> ExcludedFixturePatches;
 	FDMXFixturePatchListRowRetValDelegate IsRowEnabledDelegate;
 	FDMXFixturePatchListRowRetValDelegate IsRowVisibleDelegate;
-	FOnDragDetected OnRowDraggedDelegate;
+	FOnDragDetected OnRowDragDetectedDelegate;
 };
diff --git a/Engine/Plugins/VirtualProduction/DMX/DMXEngine/Source/DMXEditor/Public/Widgets/SDMXReadOnlyFixturePatchListRow.h b/Engine/Plugins/VirtualProduction/DMX/DMXEngine/Source/DMXEditor/Public/Widgets/SDMXReadOnlyFixturePatchListRow.h
--- a/Engine/Plugins/VirtualProduction/DMX/DMXEngine/Source/DMXEditor/Public/Widgets/SDMXReadOnlyFixturePatchListRow.h
+++ b/Engine/Plugins/VirtualProduction/DMX/DMXEngine/Source/DMXEditor/Public/Widgets/SDMXReadOnlyFixturePatchListRow.h
@@ -16,7 +16,7 @@ public:
 	SLATE_BEGIN_ARGS(SDMXReadOnlyFixturePatchListRow)
 	{}
 
-		SLATE_EVENT(FOnDragDetected, OnRowDragged)
+		SLATE_EVENT(FOnDragDetected, OnRowDragDetected)
 
 	SLATE_END_ARGS()
 			
diff --git a/Engine/Plugins/VirtualProduction/DMX/DMXPixelMapping/Source/DMXPixelMappingEditor/Private/Toolkits/DMXPixelMappingToolkit.cpp b/Engine/Plugins/VirtualProduction/DMX/DMXPixelMapping/Source/DMXPixelMappingEditor/Private/Toolkits/DMXPixelMappingToolkit.cpp
--- a/Engine/Plugins/VirtualProduction/DMX/DMXPixelMapping/Source/DMXPixelMappingEditor/Private/Toolkits/DMXPixelMappingToolkit.cpp
+++ b/Engine/Plugins/VirtualProduction/DMX/DMXPixelMapping/Source/DMXPixelMappingEditor/Private/Toolkits/DMXPixelMappingToolkit.cpp
@@ -508,47 +508,47 @@ void FDMXPixelMappingToolkit::InitializeInternal(const EToolkitMode::Type Mode,
 			(
 				FTabManager::NewSplitter()
 				->SetOrientation(Orient_Horizontal)
-				->SetSizeCoefficient(.25f)
 				->Split
 				(
 					FTabManager::NewSplitter()
 					->SetOrientation(Orient_Vertical)
+					->SetSizeCoefficient(.25f)
 					->Split
 					(
 						FTabManager::NewStack()
-						->AddTab(HierarchyViewTabID, ETabState::OpenedTab)
-						->SetSizeCoefficient(.618f)
+						->AddTab(DMXLibraryViewTabID, ETabState::OpenedTab)
+						->SetSizeCoefficient(.382f)
 					)
 					->Split
 					(
 						FTabManager::NewStack()
-						->AddTab(DMXLibraryViewTabID, ETabState::OpenedTab)
-						->SetSizeCoefficient(.382f)
+						->AddTab(HierarchyViewTabID, ETabState::OpenedTab)
+						->SetSizeCoefficient(.618f)
 					)
 				)
-				->SetSizeCoefficient(.5f)
 				->Split
 				(
 					FTabManager::NewSplitter()
 					->SetOrientation(Orient_Vertical)
+					->SetSizeCoefficient(.5f)
 					->Split
 					(
 						FTabManager::NewStack()
 						->AddTab(DesignerViewTabID, ETabState::OpenedTab)
-						->SetSizeCoefficient(.618f)
+						->SetSizeCoefficient(.75f)
 					)
 					->Split
 					(
 						FTabManager::NewStack()
 						->AddTab(PreviewViewTabID, ETabState::OpenedTab)
-						->SetSizeCoefficient(.382f)
+						->SetSizeCoefficient(.25f)
 					)
 				)
-				->SetSizeCoefficient(.25f)
 				->Split
 				(
 					FTabManager::NewSplitter()
 					->SetOrientation(Orient_Vertical)
+					->SetSizeCoefficient(.25f)
 					->Split
 					(
 						FTabManager::NewStack()
diff --git a/Engine/Plugins/VirtualProduction/DMX/DMXPixelMapping/Source/DMXPixelMappingEditor/Private/Views/SDMXPixelMappingDMXLibraryView.cpp b/Engine/Plugins/VirtualProduction/DMX/DMXPixelMapping/Source/DMXPixelMappingEditor/Private/Views/SDMXPixelMappingDMXLibraryView.cpp
--- a/Engine/Plugins/VirtualProduction/DMX/DMXPixelMapping/Source/DMXPixelMappingEditor/Private/Views/SDMXPixelMappingDMXLibraryView.cpp
+++ b/Engine/Plugins/VirtualProduction/DMX/DMXPixelMapping/Source/DMXPixelMappingEditor/Private/Views/SDMXPixelMappingDMXLibraryView.cpp
@@ -10,7 +10,6 @@
 #include "Customizations/DMXPixelMappingDetailCustomization_FixtureGroup.h"
 #include "DMXPixelMapping.h"
 #include "DMXPixelMappingComponentReference.h"
-#include "DragDrop/DMXPixelMappingDragDropOp.h"
 #include "IDetailsView.h"
 #include "Library/DMXEntityFixturePatch.h"
 #include "Library/DMXLibrary.h"
@@ -181,7 +180,7 @@ void SDMXPixelMappingDMXLibraryView::Construct(const FArguments& InArgs, const T
 				.BorderImage(FAppStyle::Get().GetBrush("Brushes.Recessed"))
 				.Padding(FMargin(8.f, 2.f, 2.f, 0.f))
 				[
-					SAssignNew(FixturePatchList, SDMXPixelMappingFixturePatchList, Toolkit)
+					SAssignNew(FixturePatchList, SDMXPixelMappingFixturePatchList, Toolkit, ViewModel)
 				]
 			]
 
@@ -379,53 +378,6 @@ FReply SDMXPixelMappingDMXLibraryView::OnAddAllPatchesClicked()
 	return FReply::Handled();
 }
 
-FReply SDMXPixelMappingDMXLibraryView::OnFixturePatchRowDragged(const FGeometry& MyGeometry, const FPointerEvent& MouseEvent)
-{
-	if (!ensureMsgf(ViewModel, TEXT("Invalid view model for PixelMapping DMX Library View, cannot display view.")))
-	{
-		return FReply::Handled();
-	}
-
-	UDMXPixelMappingFixtureGroupComponent* FixtureGroupComponent = ViewModel->GetFixtureGroupComponent();
-	if (!FixtureGroupComponent)
-	{
-		return FReply::Handled();
-	}
-
-	const TArray<TSharedPtr<FDMXEntityFixturePatchRef>> SelectedPatches = FixturePatchList->GetSelectedFixturePatchRefs();
-	if (!SelectedPatches.IsEmpty())
-	{
-		TArray<TSharedPtr<FDMXPixelMappingComponentTemplate>> Templates;
-		for (const TSharedPtr<FDMXEntityFixturePatchRef>& FixturePatchRef : SelectedPatches)
-		{
-			if (!FixturePatchRef.IsValid())
-			{
-				continue;
-			}
-
-			const UDMXEntityFixturePatch* FixturePatch = FixturePatchRef->GetFixturePatch();
-			const UDMXEntityFixtureType* FixtureType = FixturePatch ? FixturePatch->GetFixtureType() : nullptr;
-			const FDMXFixtureMode* ActiveModePtr = FixturePatch ? FixturePatch->GetActiveMode() : nullptr;
-			if (FixturePatch && FixtureType && ActiveModePtr)
-			{
-				if (ActiveModePtr->bFixtureMatrixEnabled)
-				{
-					TSharedRef<FDMXPixelMappingComponentTemplate> FixturePatchMatrixTemplate = MakeShared<FDMXPixelMappingComponentTemplate>(UDMXPixelMappingMatrixComponent::StaticClass(), *FixturePatchRef);
-					Templates.Add(FixturePatchMatrixTemplate);
-				}
-				else
-				{
-					TSharedRef<FDMXPixelMappingComponentTemplate> FixturePatchItemTemplate = MakeShared<FDMXPixelMappingComponentTemplate>(UDMXPixelMappingFixtureGroupItemComponent::StaticClass(), *FixturePatchRef);
-					Templates.Add(FixturePatchItemTemplate);
-				}
-			}
-		}
-		return FReply::Handled().BeginDragDrop(FDMXPixelMappingDragDropOp::New(FVector2D::ZeroVector, Templates, FixtureGroupComponent));
-	}
-
-	return FReply::Unhandled();
-}
-
 TArray<UDMXEntityFixturePatch*> SDMXPixelMappingDMXLibraryView::GetFixturePatchesInDMXLibrary() const
 {
 	if (!ensureMsgf(ViewModel, TEXT("Invalid view model for PixelMapping DMX Library View, cannot display view.")))
@@ -472,5 +424,4 @@ TArray<UDMXEntityFixturePatch*> SDMXPixelMappingDMXLibraryView::GetFixturePatche
 	return FixturePatchesInPixelMapping;
 }
 
-
 #undef LOCTEXT_NAMESPACE
diff --git a/Engine/Plugins/VirtualProduction/DMX/DMXPixelMapping/Source/DMXPixelMappingEditor/Private/Views/SDMXPixelMappingDMXLibraryView.h b/Engine/Plugins/VirtualProduction/DMX/DMXPixelMapping/Source/DMXPixelMappingEditor/Private/Views/SDMXPixelMappingDMXLibraryView.h
--- a/Engine/Plugins/VirtualProduction/DMX/DMXPixelMapping/Source/DMXPixelMappingEditor/Private/Views/SDMXPixelMappingDMXLibraryView.h
+++ b/Engine/Plugins/VirtualProduction/DMX/DMXPixelMapping/Source/DMXPixelMappingEditor/Private/Views/SDMXPixelMappingDMXLibraryView.h
@@ -78,9 +78,6 @@ private:
 	/** Called when the 'Add All Patches' button was clicked */
 	FReply OnAddAllPatchesClicked();
 
-	/** Called when a fixture patch row was dragged */
-	FReply OnFixturePatchRowDragged(const FGeometry& MyGeometry, const FPointerEvent& MouseEvent);
-
 	/** Helper that returns all fixture patches in use in the pixel mapping */
 	TArray<UDMXEntityFixturePatch*> GetFixturePatchesInDMXLibrary() const;
 
diff --git a/Engine/Plugins/VirtualProduction/DMX/DMXPixelMapping/Source/DMXPixelMappingEditor/Private/Views/SDMXPixelMappingHierarchyView.cpp b/Engine/Plugins/VirtualProduction/DMX/DMXPixelMapping/Source/DMXPixelMappingEditor/Private/Views/SDMXPixelMappingHierarchyView.cpp
--- a/Engine/Plugins/VirtualProduction/DMX/DMXPixelMapping/Source/DMXPixelMappingEditor/Private/Views/SDMXPixelMappingHierarchyView.cpp
+++ b/Engine/Plugins/VirtualProduction/DMX/DMXPixelMapping/Source/DMXPixelMappingEditor/Private/Views/SDMXPixelMappingHierarchyView.cpp
@@ -29,7 +29,7 @@
 #include "Widgets/Input/SComboButton.h"
 #include "Widgets/Input/SSearchBox.h"
 #include "Widgets/Layout/SScrollBorder.h"
-#include "Widgets/SDMXPixelMappingHierarchyItem.h"
+#include "Widgets/SDMXPixelMappingHierarchyRow.h"
 #include "Widgets/Views/STreeView.h"
 
 
@@ -382,7 +382,7 @@ void SDMXPixelMappingHierarchyView::OnGetChildItems(TSharedPtr<FDMXPixelMappingH
 
 TSharedRef<ITableRow> SDMXPixelMappingHierarchyView::OnGenerateRow(FDMXPixelMappingHierarchyItemWidgetModelPtr Item, const TSharedRef<STableViewBase>& OwnerTable)
 {
-	return SNew(SDMXPixelMappingHierarchyItem, OwnerTable, WeakToolkit, Item.ToSharedRef());
+	return SNew(SDMXPixelMappingHierarchyRow, OwnerTable, WeakToolkit, Item.ToSharedRef());
 }
 
 TSharedPtr<SWidget> SDMXPixelMappingHierarchyView::OnContextMenuOpening()
@@ -572,7 +572,7 @@ void SDMXPixelMappingHierarchyView::RenameSelectedComponent()
 		return;
 	}
 
-	StaticCastSharedPtr<SDMXPixelMappingHierarchyItem>(TableRow)->EnterRenameMode();
+	StaticCastSharedPtr<SDMXPixelMappingHierarchyRow>(TableRow)->EnterRenameMode();
 }
 
 void SDMXPixelMappingHierarchyView::DeleteSelectedComponents()
diff --git a/Engine/Plugins/VirtualProduction/DMX/DMXPixelMapping/Source/DMXPixelMappingEditor/Private/Widgets/SDMXPixelMappingFixturePatchList.cpp b/Engine/Plugins/VirtualProduction/DMX/DMXPixelMapping/Source/DMXPixelMappingEditor/Private/Widgets/SDMXPixelMappingFixturePatchList.cpp
--- a/Engine/Plugins/VirtualProduction/DMX/DMXPixelMapping/Source/DMXPixelMappingEditor/Private/Widgets/SDMXPixelMappingFixturePatchList.cpp
+++ b/Engine/Plugins/VirtualProduction/DMX/DMXPixelMapping/Source/DMXPixelMappingEditor/Private/Widgets/SDMXPixelMappingFixturePatchList.cpp
@@ -4,26 +4,74 @@
 
 #include "Algo/MaxElement.h"
 #include "DMXPixelMapping.h"
+#include "DragDrop/DMXPixelMappingDragDropOp.h"
 #include "Components/DMXPixelMappingBaseComponent.h"
 #include "Components/DMXPixelMappingFixtureGroupItemComponent.h"
 #include "Components/DMXPixelMappingMatrixComponent.h"
 #include "Library/DMXEntityFixturePatch.h"
 #include "Library/DMXLibrary.h"
+#include "Templates/DMXPixelMappingComponentTemplate.h"
 #include "Toolkits/DMXPixelMappingToolkit.h"
+#include "ViewModels/DMXPixelMappingDMXLibraryViewModel.h"
 
 
-void SDMXPixelMappingFixturePatchList::Construct(const FArguments& InArgs, const TSharedPtr<FDMXPixelMappingToolkit>& InToolkit)
+void SDMXPixelMappingFixturePatchList::Construct(const FArguments& InArgs, const TSharedPtr<FDMXPixelMappingToolkit>& InToolkit, TWeakObjectPtr<UDMXPixelMappingDMXLibraryViewModel> InDMXLibraryModel)
 {
-	if (!InToolkit.IsValid())
+	if (!InToolkit.IsValid() || !InDMXLibraryModel.IsValid())
 	{
 		return;
 	}
 	WeakToolkit = InToolkit;
+	WeakDMXLibraryViewModel = InDMXLibraryModel;
 
-	SDMXReadOnlyFixturePatchList::Construct(SDMXReadOnlyFixturePatchList::FArguments()
+	SDMXReadOnlyFixturePatchList::Construct(
+		SDMXReadOnlyFixturePatchList::FArguments()
 		.ListDescriptor(InArgs._ListDescriptor)
 		.OnContextMenuOpening(InArgs._OnContextMenuOpening)
-		.OnRowDragged(InArgs._OnRowDragged));
+		.OnRowDragDetected(this, &SDMXPixelMappingFixturePatchList::OnRowDragDetected)
+	);
+}
+
+FReply SDMXPixelMappingFixturePatchList::OnRowDragDetected(const FGeometry& MyGeometry, const FPointerEvent& MouseEvent)
+{
+	const TSharedPtr<FDMXPixelMappingToolkit> Toolkit = WeakToolkit.Pin();
+	UDMXPixelMappingFixtureGroupComponent* FixtureGroup = WeakDMXLibraryViewModel.IsValid() ? WeakDMXLibraryViewModel->GetFixtureGroupComponent() : nullptr;
+	if (!Toolkit.IsValid() || !FixtureGroup)
+	{
+		return FReply::Unhandled();
+	}
+
+	TArray<TSharedPtr<FDMXPixelMappingComponentTemplate>> Templates;
+	for (const TSharedPtr<FDMXEntityFixturePatchRef>& FixturePatchRef : GetSelectedFixturePatchRefs())
+	{
+		UDMXEntityFixturePatch* FixturePatch = FixturePatchRef.IsValid() ? FixturePatchRef->GetFixturePatch() : nullptr;
+		UDMXEntityFixtureType* FixtureType = FixturePatch ? FixturePatch->GetFixtureType() : nullptr;
+		if (!FixturePatch || !FixtureType)
+		{
+			continue;
+		}
+
+		const FDMXEntityFixturePatchRef FixturePatchReference(FixturePatch);
+		if (FixtureType->bFixtureMatrixEnabled)
+		{
+			TSharedRef<FDMXPixelMappingComponentTemplate> FixturePatchMatrixTemplate = MakeShared<FDMXPixelMappingComponentTemplate>(UDMXPixelMappingMatrixComponent::StaticClass(), FixturePatchReference);
+			Templates.Add(FixturePatchMatrixTemplate);
+		}
+		else
+		{
+			TSharedRef<FDMXPixelMappingComponentTemplate> FixturePatchItemTemplate = MakeShared<FDMXPixelMappingComponentTemplate>(UDMXPixelMappingFixtureGroupItemComponent::StaticClass(), FixturePatchReference);
+			Templates.Add(FixturePatchItemTemplate);
+		}
+	}
+
+	if (Templates.IsEmpty())
+	{
+		return FReply::Handled();
+	}
+	else
+	{
+		return FReply::Handled().BeginDragDrop(FDMXPixelMappingDragDropOp::New(FVector2D::ZeroVector, Templates, FixtureGroup));
+	}
 }
 
 void SDMXPixelMappingFixturePatchList::SelectAfter(const TArray<TSharedPtr<FDMXEntityFixturePatchRef>>& FixturePatches)
diff --git a/Engine/Plugins/VirtualProduction/DMX/DMXPixelMapping/Source/DMXPixelMappingEditor/Private/Widgets/SDMXPixelMappingFixturePatchList.h b/Engine/Plugins/VirtualProduction/DMX/DMXPixelMapping/Source/DMXPixelMappingEditor/Private/Widgets/SDMXPixelMappingFixturePatchList.h
--- a/Engine/Plugins/VirtualProduction/DMX/DMXPixelMapping/Source/DMXPixelMappingEditor/Private/Widgets/SDMXPixelMappingFixturePatchList.h
+++ b/Engine/Plugins/VirtualProduction/DMX/DMXPixelMapping/Source/DMXPixelMappingEditor/Private/Widgets/SDMXPixelMappingFixturePatchList.h
@@ -7,6 +7,7 @@
 
 class FDMXPixelMappingToolkit;
 class UDMXEntityFixturePatch;
+class UDMXPixelMappingDMXLibraryViewModel;
 
 
 /** Displays the DMX Library of the currently selected fixture group component */
@@ -22,13 +23,10 @@ public:
 		/** Called when a row of the list is right clicked */
 		SLATE_EVENT(FOnContextMenuOpening, OnContextMenuOpening)
 
-		/** Delegate executed when a row was dragged */
-		SLATE_EVENT(FOnDragDetected, OnRowDragged)
-
 	SLATE_END_ARGS()
 
 	/** Constructs this widget */
-	void Construct(const FArguments& InArgs, const TSharedPtr<FDMXPixelMappingToolkit>& InToolkit);
+	void Construct(const FArguments& InArgs, const TSharedPtr<FDMXPixelMappingToolkit>& InToolkit, TWeakObjectPtr<UDMXPixelMappingDMXLibraryViewModel> InDMXLibraryModel);
 
 	/** Selects the first patch after the specified fixture patches */
 	void SelectAfter(const TArray<TSharedPtr<FDMXEntityFixturePatchRef>>& FixturePatches);
@@ -38,9 +36,15 @@ private:
 	virtual void RefreshList() override;
 	//~ End SDMXReadOnlyFixturePatchList interface
 
+	/** Called when a row in the list was dragged */
+	FReply OnRowDragDetected(const FGeometry& MyGeometry, const FPointerEvent& MouseEvent);
+
 	/** Holds the fixutre patches which are hidden the list */
 	TArray<TSharedPtr<FDMXEntityFixturePatchRef>> HiddenFixturePatches;
 
+	/** View model of the displayed dmx library */
+	TWeakObjectPtr<UDMXPixelMappingDMXLibraryViewModel> WeakDMXLibraryViewModel;
+
 	/** The toolkit of the editor that displays this widget */
 	TWeakPtr<FDMXPixelMappingToolkit> WeakToolkit;
 };
diff --git a/Engine/Plugins/VirtualProduction/DMX/DMXPixelMapping/Source/DMXPixelMappingEditor/Private/Widgets/SDMXPixelMappingHierarchyItem.cpp b/Engine/Plugins/VirtualProduction/DMX/DMXPixelMapping/Source/DMXPixelMappingEditor/Private/Widgets/SDMXPixelMappingHierarchyItem.cpp
deleted file mode 100644
diff --git a/Engine/Plugins/VirtualProduction/DMX/DMXPixelMapping/Source/DMXPixelMappingEditor/Private/Widgets/SDMXPixelMappingHierarchyItem.h b/Engine/Plugins/VirtualProduction/DMX/DMXPixelMapping/Source/DMXPixelMappingEditor/Private/Widgets/SDMXPixelMappingHierarchyItem.h
deleted file mode 100644
diff --git a/Engine/Plugins/VirtualProduction/DMX/DMXPixelMapping/Source/DMXPixelMappingEditor/Private/Widgets/SDMXPixelMappingHierarchyRow.cpp b/Engine/Plugins/VirtualProduction/DMX/DMXPixelMapping/Source/DMXPixelMappingEditor/Private/Widgets/SDMXPixelMappingHierarchyRow.cpp
new file mode 100644
--- /dev/null
+++ b/Engine/Plugins/VirtualProduction/DMX/DMXPixelMapping/Source/DMXPixelMappingEditor/Private/Widgets/SDMXPixelMappingHierarchyRow.cpp
@@ -0,0 +1,182 @@
+// Copyright Epic Games, Inc. All Rights Reserved.
+
+#include "Widgets/SDMXPixelMappingHierarchyRow.h"
+
+#include "Components/DMXPixelMappingOutputComponent.h"
+#include "DMXEditorStyle.h"
+#include "DMXPixelMappingEditorUtils.h"
+#include "DragDrop/DMXPixelMappingDragDropOp.h"
+#include "ScopedTransaction.h"
+#include "Toolkits/DMXPixelMappingToolkit.h"
+#include "Widgets/Text/SInlineEditableTextBlock.h"
+#include "Widgets/Views/SListView.h"
+#include "ViewModels/DMXPixelMappingHierarchyItem.h"
+#include "Views/SDMXPixelMappingHierarchyView.h"
+
+
+#define LOCTEXT_NAMESPACE "SDMXPixelMappingHierarchyRow"
+
+void SDMXPixelMappingHierarchyRow::Construct(const FArguments& InArgs, const TSharedRef<STableViewBase>& InOwnerTableView, TWeakPtr<FDMXPixelMappingToolkit> InWeakToolkit, const TSharedRef<FDMXPixelMappingHierarchyItem>& InItem)
+{
+	WeakToolkit = InWeakToolkit;
+	Item = InItem;
+
+	SMultiColumnTableRow<FDMXPixelMappingHierarchyItemWidgetModelPtr>::Construct(
+		FSuperRowType::FArguments()
+		.Padding(0.0f)
+		.Style(&FAppStyle::Get().GetWidgetStyle<FTableRowStyle>("SimpleTableView.Row"))
+		.OnDragDetected(this, &SDMXPixelMappingHierarchyRow::OnRowDragDetected),
+		InOwnerTableView);
+}
+
+void SDMXPixelMappingHierarchyRow::EnterRenameMode()
+{
+	EditableNameTextBox->EnterEditingMode();
+}
+
+TSharedRef<SWidget> SDMXPixelMappingHierarchyRow::GenerateWidgetForColumn(const FName& ColumnName)
+{
+	if (ColumnName == SDMXPixelMappingHierarchyView::FColumnIds::EditorColor)
+	{
+		return GenerateEditorColorWidget();
+	}
+	else if (ColumnName == SDMXPixelMappingHierarchyView::FColumnIds::ComponentName)
+	{			
+		// The name column gets the tree expansion arrow
+		return SNew(SBox)
+			[
+				SNew(SHorizontalBox)
+
+				+ SHorizontalBox::Slot()
+				.AutoWidth()
+				[
+					SNew(SExpanderArrow, SharedThis(this))
+				]
+
+				+ SHorizontalBox::Slot()
+				.FillWidth(1.f)
+				[
+					GenerateComponentNameWidget()
+				]
+			];
+	}
+	else if (ColumnName == SDMXPixelMappingHierarchyView::FColumnIds::FixtureID)
+	{
+		return GenerateFixtureIDWidget();
+	}
+	else if (ColumnName == SDMXPixelMappingHierarchyView::FColumnIds::Patch)
+	{
+		return GeneratePatchWidget();
+	}
+
+	return SNullWidget::NullWidget;
+}
+
+TSharedRef<SWidget> SDMXPixelMappingHierarchyRow::GenerateEditorColorWidget()
+{
+	if (Item.IsValid() && Item->HasEditorColor())
+	{
+		return
+			SNew(SBorder)
+			.HAlign(HAlign_Fill)
+			.Padding(5.f, 2.f)
+			.BorderImage(FAppStyle::GetBrush("NoBorder"))
+			[
+				SAssignNew(EditorColorImageWidget, SImage)
+				.Image(FDMXEditorStyle::Get().GetBrush("DMXEditor.WhiteRoundedPropertyBorder"))
+			.ColorAndOpacity_Lambda([this]() -> FSlateColor
+				{
+					return Item.IsValid() ? Item->GetEditorColor() : FLinearColor::Red;
+				})
+			];
+	}
+
+	return SNullWidget::NullWidget;
+}
+
+TSharedRef<SWidget> SDMXPixelMappingHierarchyRow::GenerateComponentNameWidget()
+{
+	return
+		SNew(SBorder)
+		.HAlign(HAlign_Fill)
+		.VAlign(VAlign_Center)
+		.Padding(4.f)
+		.BorderImage(FAppStyle::GetBrush("NoBorder"))
+		[
+			SAssignNew(EditableNameTextBox, SInlineEditableTextBlock)
+			.Font(FAppStyle::GetFontStyle(TEXT("PropertyWindow.NormalFont")))
+			.Text(Item.Get(), &FDMXPixelMappingHierarchyItem::GetComponentNameText)
+			.ColorAndOpacity(FLinearColor::White)
+		];
+}
+
+TSharedRef<SWidget> SDMXPixelMappingHierarchyRow::GenerateFixtureIDWidget()
+{
+	return
+		SNew(SBorder)
+		.HAlign(HAlign_Fill)
+		.VAlign(VAlign_Center)
+		.Padding(4.f)
+		.BorderImage(FAppStyle::GetBrush("NoBorder"))
+		[
+			SNew(STextBlock)
+			.Font(FAppStyle::GetFontStyle(TEXT("PropertyWindow.NormalFont")))
+			.Text(Item.Get(), &FDMXPixelMappingHierarchyItem::GetFixtureIDText)
+			.ColorAndOpacity(FLinearColor::White)
+		];
+}
+
+TSharedRef<SWidget> SDMXPixelMappingHierarchyRow::GeneratePatchWidget()
+{
+	return
+		SNew(SBorder)
+		.HAlign(HAlign_Fill)
+		.VAlign(VAlign_Center)
+		.Padding(4.f)
+		.BorderImage(FAppStyle::GetBrush("NoBorder"))
+		[
+			SNew(STextBlock)
+			.Font(FAppStyle::GetFontStyle(TEXT("PropertyWindow.NormalFont")))
+			.Text(Item.Get(), &FDMXPixelMappingHierarchyItem::GetPatchText)
+			.ColorAndOpacity(FLinearColor::White)
+		];
+}
+
+void SDMXPixelMappingHierarchyRow::OnNameTextCommited(const FText& InText, ETextCommit::Type CommitInfo)
+{
+	const TSharedPtr<FDMXPixelMappingToolkit> Toolkit = WeakToolkit.Pin();
+	const UDMXPixelMappingBaseComponent* Component = Item.IsValid() ? Item->GetComponent() : nullptr;
+	if (Toolkit.IsValid() && Component && !InText.IsEmpty())
+	{
+		const FScopedTransaction RenameComponentTransaction(LOCTEXT("RenameComponentTransaction", "Rename Pixel Mapping Component"));
+
+		const FName& ComponentName = Component->GetFName();
+		Toolkit->RenameComponent(ComponentName, InText.ToString());
+	}
+}
+
+FReply SDMXPixelMappingHierarchyRow::OnRowDragDetected(const FGeometry& MyGeometry, const FPointerEvent& MouseEvent)
+{
+	// SMultiColumnTableRow for some reason does not provide the mouse key, so testing for LeftMouseButton is not required.
+	if (!WeakToolkit.IsValid())
+	{
+		return FReply::Unhandled();
+	}
+
+	const TSet<FDMXPixelMappingComponentReference>& SelectedComponents = WeakToolkit.Pin()->GetSelectedComponents();
+	TArray<TWeakObjectPtr<UDMXPixelMappingBaseComponent>> DraggedComponents;
+	Algo::TransformIf(SelectedComponents, DraggedComponents,
+		[](const FDMXPixelMappingComponentReference& ComponentReference)
+		{
+			return ComponentReference.IsValid();
+		},
+		[](const FDMXPixelMappingComponentReference& ComponentReference)
+		{
+			return ComponentReference.GetComponent();
+		}
+	);
+
+	return FReply::Handled().BeginDragDrop(FDMXPixelMappingDragDropOp::New(FVector2D::ZeroVector, DraggedComponents));
+}
+
+#undef LOCTEXT_NAMESPACE
diff --git a/Engine/Plugins/VirtualProduction/DMX/DMXPixelMapping/Source/DMXPixelMappingEditor/Private/Widgets/SDMXPixelMappingHierarchyRow.h b/Engine/Plugins/VirtualProduction/DMX/DMXPixelMapping/Source/DMXPixelMappingEditor/Private/Widgets/SDMXPixelMappingHierarchyRow.h
new file mode 100644
--- /dev/null
+++ b/Engine/Plugins/VirtualProduction/DMX/DMXPixelMapping/Source/DMXPixelMappingEditor/Private/Widgets/SDMXPixelMappingHierarchyRow.h
@@ -0,0 +1,69 @@
+// Copyright Epic Games, Inc. All Rights Reserved.
+
+#pragma once
+
+#include "DMXPixelMappingEditorCommon.h"
+#include "Widgets/Views/ITableRow.h"
+#include "Widgets/Views/STableRow.h"
+
+class FDMXPixelMappingHierarchyItem;
+class SInlineEditableTextBlock;
+class SDMXPixelMappingHierarchyView;
+class SImage;
+class SInlineEditableTextBox;
+class STableViewBase;
+
+
+/** Displays a single item in a hierarchy view */
+class SDMXPixelMappingHierarchyRow final
+	: public SMultiColumnTableRow<TSharedPtr<FDMXPixelMappingHierarchyItem>>
+{
+public:
+	SLATE_BEGIN_ARGS(SDMXPixelMappingHierarchyRow) {}
+
+		/** Delegate executed when a row was dragged */
+		SLATE_EVENT(FOnDragDetected, OnItemDragDetected)
+
+	SLATE_END_ARGS()
+
+	/** Constructs this widget */
+	void Construct(const FArguments& InArgs, const TSharedRef<STableViewBase>& InOwnerTableView, TWeakPtr<FDMXPixelMappingToolkit> InWeakToolkit, const TSharedRef<FDMXPixelMappingHierarchyItem>& InItem);
+
+	/** Starts editing the name of the component of this row */
+	void EnterRenameMode();
+
+private:
+	//~ Begin SMultiColumnTableRow interface
+	virtual TSharedRef<SWidget> GenerateWidgetForColumn(const FName& ColumnName) override;
+	//~ End SMultiColumnTableRow interface
+
+	/** Generates a widget to display the editor color */
+	TSharedRef<SWidget> GenerateEditorColorWidget();
+
+	/** Generates a widget to display the name */
+	TSharedRef<SWidget> GenerateComponentNameWidget();
+
+	/** Generates a widget to display the fixture ID */
+	TSharedRef<SWidget> GenerateFixtureIDWidget();
+
+	/** Generates a widget to display the patch */
+	TSharedRef<SWidget> GeneratePatchWidget();
+
+	/** Called when text is committed on the node */
+	void OnNameTextCommited(const FText& InText, ETextCommit::Type CommitInfo);
+
+	/** Called when this row was dragged */
+	FReply OnRowDragDetected(const FGeometry& MyGeometry, const FPointerEvent& MouseEvent);
+
+	/** Model of the component being displayed */
+	TSharedPtr<FDMXPixelMappingHierarchyItem> Item;
+
+	/** The editor color image widget */
+	TSharedPtr<SImage> EditorColorImageWidget;
+
+	/** Text box that displays and allows to edit the name */
+	TSharedPtr<SInlineEditableTextBlock> EditableNameTextBox;
+
+	/** Weak toolkit this widget belongs to */
+	TWeakPtr<FDMXPixelMappingToolkit> WeakToolkit;
+};
diff --git a/Engine/Plugins/VirtualProduction/DMX/DMXPixelMapping/Source/DMXPixelMappingEditor/Private/Widgets/SDMXPixelMappingSurface.cpp b/Engine/Plugins/VirtualProduction/DMX/DMXPixelMapping/Source/DMXPixelMappingEditor/Private/Widgets/SDMXPixelMappingSurface.cpp
--- a/Engine/Plugins/VirtualProduction/DMX/DMXPixelMapping/Source/DMXPixelMappingEditor/Private/Widgets/SDMXPixelMappingSurface.cpp
+++ b/Engine/Plugins/VirtualProduction/DMX/DMXPixelMapping/Source/DMXPixelMappingEditor/Private/Widgets/SDMXPixelMappingSurface.cpp
@@ -42,6 +42,7 @@ struct FFixedZoomLevelsContainerDesignSurface : public FZoomLevelsContainer
 		ZoomLevels.Add(FZoomLevelEntry(0.25f, FText::FromString("1:4"), EGraphRenderingLOD::MediumDetail));
 		ZoomLevels.Add(FZoomLevelEntry(0.333333333333f, FText::FromString("1:3"), EGraphRenderingLOD::MediumDetail));
 		ZoomLevels.Add(FZoomLevelEntry(0.5f, FText::FromString("1:2"), EGraphRenderingLOD::MediumDetail));
+		ZoomLevels.Add(FZoomLevelEntry(0.75f, FText::FromString("1:1.5"), EGraphRenderingLOD::MediumDetail));
 		ZoomLevels.Add(FZoomLevelEntry(1.000f, FText::FromString("1:1"), EGraphRenderingLOD::DefaultDetail));
 		ZoomLevels.Add(FZoomLevelEntry(1.250f, FText::FromString("1.25:1"), EGraphRenderingLOD::DefaultDetail));
 		ZoomLevels.Add(FZoomLevelEntry(1.500f, FText::FromString("1.5:1"), EGraphRenderingLOD::DefaultDetail));
