DMX - Fix Pixel Mapping Components are not always z-ordered correctly, hierarchy view follows selection correctly

diff --git a/Engine/Plugins/VirtualProduction/DMX/DMXPixelMapping/Source/DMXPixelMappingEditor/Private/DragDrop/DMXPixelMappingDragDropOp.cpp b/Engine/Plugins/VirtualProduction/DMX/DMXPixelMapping/Source/DMXPixelMappingEditor/Private/DragDrop/DMXPixelMappingDragDropOp.cpp
--- a/Engine/Plugins/VirtualProduction/DMX/DMXPixelMapping/Source/DMXPixelMappingEditor/Private/DragDrop/DMXPixelMappingDragDropOp.cpp
+++ b/Engine/Plugins/VirtualProduction/DMX/DMXPixelMapping/Source/DMXPixelMappingEditor/Private/DragDrop/DMXPixelMappingDragDropOp.cpp
@@ -13,16 +13,6 @@
 
 FDMXPixelMappingDragDropOp::~FDMXPixelMappingDragDropOp()
 {
-	// Newly added components are highest Z-Order
-	TArray<UDMXPixelMappingOutputComponent*> ChildrenNotOverTheirParent;
-	for (const TWeakObjectPtr<UDMXPixelMappingBaseComponent>& DraggedComponent : DraggedComponents)
-	{
-		if (UDMXPixelMappingOutputComponent* DraggedOutputComponent = Cast<UDMXPixelMappingOutputComponent>(DraggedComponent.Get()))
-		{
-			DraggedOutputComponent->MakeHighestZOrderInComponentRect();
-		}
-	}
-
 	// End the transaction in any case
 	GEditor->EndTransaction();
 }
diff --git a/Engine/Plugins/VirtualProduction/DMX/DMXPixelMapping/Source/DMXPixelMappingEditor/Private/Toolkits/DMXPixelMappingToolkit.cpp b/Engine/Plugins/VirtualProduction/DMX/DMXPixelMapping/Source/DMXPixelMappingEditor/Private/Toolkits/DMXPixelMappingToolkit.cpp
--- a/Engine/Plugins/VirtualProduction/DMX/DMXPixelMapping/Source/DMXPixelMappingEditor/Private/Toolkits/DMXPixelMappingToolkit.cpp
+++ b/Engine/Plugins/VirtualProduction/DMX/DMXPixelMapping/Source/DMXPixelMappingEditor/Private/Toolkits/DMXPixelMappingToolkit.cpp
@@ -2,6 +2,8 @@
 
 #include "Toolkits/DMXPixelMappingToolkit.h"
 
+#include "Algo/Sort.h"
+#include "Algo/Transform.h"
 #include "DMXPixelMapping.h"
 #include "DMXPixelMappingEditorCommands.h"
 #include "DMXPixelMappingEditorModule.h"
@@ -227,6 +229,23 @@ void FDMXPixelMappingToolkit::SelectComponents(const TSet<FDMXPixelMappingCompon
 		}
 	}
 
+	// Always order selected components topmost, but keep their relative z-ordering
+	TArray<UDMXPixelMappingOutputComponent*> SelectedOutputComponents;
+	Algo::TransformIf(SelectedComponents, SelectedOutputComponents,
+		[](const FDMXPixelMappingComponentReference& ComponentReference)
+		{
+			return ComponentReference.GetComponent() && ComponentReference.GetComponent()->IsA(UDMXPixelMappingOutputComponent::StaticClass());
+		},
+		[](const FDMXPixelMappingComponentReference& ComponentReference)
+		{
+			return CastChecked<UDMXPixelMappingOutputComponent>(ComponentReference.GetComponent());
+		});
+	Algo::SortBy(SelectedOutputComponents, &UDMXPixelMappingOutputComponent::GetZOrder);
+	for (UDMXPixelMappingOutputComponent* SelectedComponent : SelectedOutputComponents)
+	{
+		SelectedComponent->ZOrderTopmost();
+	}
+
 	OnSelectedComponentsChangedDelegate.Broadcast();
 }
 
diff --git a/Engine/Plugins/VirtualProduction/DMX/DMXPixelMapping/Source/DMXPixelMappingEditor/Private/ViewModels/DMXPixelMappingDMXLibraryViewModel.cpp b/Engine/Plugins/VirtualProduction/DMX/DMXPixelMapping/Source/DMXPixelMappingEditor/Private/ViewModels/DMXPixelMappingDMXLibraryViewModel.cpp
--- a/Engine/Plugins/VirtualProduction/DMX/DMXPixelMapping/Source/DMXPixelMappingEditor/Private/ViewModels/DMXPixelMappingDMXLibraryViewModel.cpp
+++ b/Engine/Plugins/VirtualProduction/DMX/DMXPixelMapping/Source/DMXPixelMappingEditor/Private/ViewModels/DMXPixelMappingDMXLibraryViewModel.cpp
@@ -70,6 +70,7 @@ void UDMXPixelMappingDMXLibraryViewModel::CreateAndSetNewFixtureGroup(TWeakPtr<F
 
 		NewFixtureGroupComponent->SetPosition(NewPosition);
 		NewFixtureGroupComponent->SetSize(NewSize);
+		NewFixtureGroupComponent->ZOrderTopmost();
 	}
 
 	UpdateFixtureGroupFromSelection(InWeakToolkit);
@@ -187,7 +188,7 @@ void UDMXPixelMappingDMXLibraryViewModel::AddFixturePatchesEnsured(const TArray<
 void UDMXPixelMappingDMXLibraryViewModel::SaveFixturePatchListDescriptor(const FDMXReadOnlyFixturePatchListDescriptor& NewDescriptor)
 {
 	FixturePatchListDescriptor = NewDescriptor;
-	SaveConfig();
+	SaveConfig();	
 }
 
 void UDMXPixelMappingDMXLibraryViewModel::PostUndo(bool bSuccess)
diff --git a/Engine/Plugins/VirtualProduction/DMX/DMXPixelMapping/Source/DMXPixelMappingEditor/Private/ViewModels/DMXPixelMappingOutputComponentModel.cpp b/Engine/Plugins/VirtualProduction/DMX/DMXPixelMapping/Source/DMXPixelMappingEditor/Private/ViewModels/DMXPixelMappingOutputComponentModel.cpp
--- a/Engine/Plugins/VirtualProduction/DMX/DMXPixelMapping/Source/DMXPixelMappingEditor/Private/ViewModels/DMXPixelMappingOutputComponentModel.cpp
+++ b/Engine/Plugins/VirtualProduction/DMX/DMXPixelMapping/Source/DMXPixelMappingEditor/Private/ViewModels/DMXPixelMappingOutputComponentModel.cpp
@@ -166,7 +166,7 @@ FLinearColor FDMXPixelMappingOutputComponentModel::GetColor() const
 			}
 			else
 			{
-				return OutputComponent->GetEditorColor().CopyWithNewOpacity(0.6f);
+				return OutputComponent->GetEditorColor().CopyWithNewOpacity(0.75f);
 			}
 		}
 	}
diff --git a/Engine/Plugins/VirtualProduction/DMX/DMXPixelMapping/Source/DMXPixelMappingEditor/Private/Views/SDMXPixelMappingHierarchyView.cpp b/Engine/Plugins/VirtualProduction/DMX/DMXPixelMapping/Source/DMXPixelMappingEditor/Private/Views/SDMXPixelMappingHierarchyView.cpp
--- a/Engine/Plugins/VirtualProduction/DMX/DMXPixelMapping/Source/DMXPixelMappingEditor/Private/Views/SDMXPixelMappingHierarchyView.cpp
+++ b/Engine/Plugins/VirtualProduction/DMX/DMXPixelMapping/Source/DMXPixelMappingEditor/Private/Views/SDMXPixelMappingHierarchyView.cpp
@@ -513,44 +513,31 @@ void SDMXPixelMappingHierarchyView::AdoptSelectionFromToolkit()
 		return;
 	}
 
-	if (HierarchyTreeView.IsValid())
-	{
-		HierarchyTreeView->ClearSelection();
-	}
-
-	for (const TSharedPtr<FDMXPixelMappingHierarchyItem>& Item : FilteredRootItems)
+	for (const TSharedPtr<FDMXPixelMappingHierarchyItem>& RootItem : AllRootItems)
 	{
-		RecursiveAdoptSelectionFromToolkit(Item.ToSharedRef());
+		RecursiveAdoptSelectionFromToolkit(RootItem.ToSharedRef());
 	}
 }
 
-bool SDMXPixelMappingHierarchyView::RecursiveAdoptSelectionFromToolkit(const TSharedRef<FDMXPixelMappingHierarchyItem>& Item)
+void SDMXPixelMappingHierarchyView::RecursiveAdoptSelectionFromToolkit(const TSharedRef<FDMXPixelMappingHierarchyItem>& Item)
 {
-	if (bIsUpdatingSelection)
-	{
-		return false;
-	}
-
-	const TSharedPtr<FDMXPixelMappingToolkit> Toolkit = WeakToolkit.Pin();
-	if (!Toolkit.IsValid())
-	{
-		return false;
-	}
-
-	bool bContainsSelection = false;
-	for (const TSharedPtr<FDMXPixelMappingHierarchyItem>& ChildItem : Item->GetChildren())
+	if (!WeakToolkit.IsValid() || !HierarchyTreeView.IsValid())
 	{
-		bContainsSelection |= RecursiveAdoptSelectionFromToolkit(ChildItem.ToSharedRef());
+		return;
 	}
 
-	const TSet<FDMXPixelMappingComponentReference> SelectedComponents = Toolkit->GetSelectedComponents();
-	if (SelectedComponents.Contains(FDMXPixelMappingComponentReference(Toolkit, Item->GetComponent())))
+	for (const TSharedPtr<FDMXPixelMappingHierarchyItem>& Child : Item->GetChildren())
 	{
-		HierarchyTreeView->RequestScrollIntoView(Item);
-		return true;
+		RecursiveAdoptSelectionFromToolkit(Child.ToSharedRef());
 	}
 
-	return bContainsSelection;
+	const bool bSelected = Algo::FindByPredicate(WeakToolkit.Pin()->GetSelectedComponents(),
+		[Item](const FDMXPixelMappingComponentReference& ComponentReference)
+		{
+			return ComponentReference.GetComponent() && ComponentReference.GetComponent() == Item->GetComponent();
+		}) != nullptr;
+	
+	HierarchyTreeView->SetItemSelection(Item, bSelected);
 }
 
 bool SDMXPixelMappingHierarchyView::CanRenameSelectedComponent() const
diff --git a/Engine/Plugins/VirtualProduction/DMX/DMXPixelMapping/Source/DMXPixelMappingEditor/Private/Views/SDMXPixelMappingHierarchyView.h b/Engine/Plugins/VirtualProduction/DMX/DMXPixelMapping/Source/DMXPixelMappingEditor/Private/Views/SDMXPixelMappingHierarchyView.h
--- a/Engine/Plugins/VirtualProduction/DMX/DMXPixelMapping/Source/DMXPixelMappingEditor/Private/Views/SDMXPixelMappingHierarchyView.h
+++ b/Engine/Plugins/VirtualProduction/DMX/DMXPixelMapping/Source/DMXPixelMappingEditor/Private/Views/SDMXPixelMappingHierarchyView.h
@@ -110,7 +110,7 @@ private:
 	void AdoptSelectionFromToolkit();
 
 	/** Adopts the selection from toolkit */
-	bool RecursiveAdoptSelectionFromToolkit(const TSharedRef<FDMXPixelMappingHierarchyItem>& Model);
+	void RecursiveAdoptSelectionFromToolkit(const TSharedRef<FDMXPixelMappingHierarchyItem>& Model);
 
 	/** Returns true if the current selection can to be renamed */
 	bool CanRenameSelectedComponent() const;
diff --git a/Engine/Plugins/VirtualProduction/DMX/DMXPixelMapping/Source/DMXPixelMappingRuntime/Private/Components/DMXPixelMappingOutputComponent.cpp b/Engine/Plugins/VirtualProduction/DMX/DMXPixelMapping/Source/DMXPixelMappingRuntime/Private/Components/DMXPixelMappingOutputComponent.cpp
--- a/Engine/Plugins/VirtualProduction/DMX/DMXPixelMapping/Source/DMXPixelMappingRuntime/Private/Components/DMXPixelMappingOutputComponent.cpp
+++ b/Engine/Plugins/VirtualProduction/DMX/DMXPixelMapping/Source/DMXPixelMappingRuntime/Private/Components/DMXPixelMappingOutputComponent.cpp
@@ -312,35 +312,42 @@ UDMXPixelMappingRendererComponent* UDMXPixelMappingOutputComponent::FindRenderer
 }
 
 #if WITH_EDITOR
-void UDMXPixelMappingOutputComponent::MakeHighestZOrderInComponentRect()
+void UDMXPixelMappingOutputComponent::ZOrderTopmost()
 {
-	if (UDMXPixelMappingRendererComponent* RendererComponent = FindRendererComponent())
+	UDMXPixelMappingRendererComponent* RendererComponent = FindRendererComponent();
+	if (!RendererComponent)
 	{
-		constexpr bool bRecursive = true;
-		RendererComponent->ForEachChildOfClass<UDMXPixelMappingOutputComponent>([this](UDMXPixelMappingOutputComponent* InComponent)
-			{
-				if (InComponent == this)
-				{
-					return;
-				}
+		return;
+	}
 
-				// Exclude children, they're updated when SetZOrder is called below
-				for (UDMXPixelMappingBaseComponent* OtherParent = InComponent->GetParent(); OtherParent; OtherParent = OtherParent->GetParent())
-				{
-					if (OtherParent == this)
-					{
-						return;
-					}
-				}
+	// Gather all components, sort them by zorder and rebase all
+	constexpr bool bRecursive = true;
 
-				if (this->OverlapsComponent(InComponent))
-				{
-					if (InComponent->ZOrder + 1 > ZOrder)
-					{
-						SetZOrder(InComponent->ZOrder + 1);
-					}
-				}
-			}, bRecursive);
+	TArray<UDMXPixelMappingOutputComponent*> AllOutputComoponents;
+	RendererComponent->ForEachChildOfClass<UDMXPixelMappingOutputComponent>([&AllOutputComoponents](UDMXPixelMappingOutputComponent* Component)
+		{
+			AllOutputComoponents.Add(Component);
+		}, bRecursive);
+	Algo::SortBy(AllOutputComoponents, &UDMXPixelMappingOutputComponent::GetZOrder);
+
+	int32 NewZOrder = 0;
+	for (UDMXPixelMappingOutputComponent* OutputComponent : AllOutputComoponents)
+	{
+		OutputComponent->SetZOrder(++NewZOrder);
+	}
+
+	// Order this component and its children topmost
+	SetZOrder(++NewZOrder);
+	for (UDMXPixelMappingOutputComponent* OutputComponent : AllOutputComoponents)
+	{
+		// Increment zorder for all children that are a sibling of this component
+		for (UDMXPixelMappingBaseComponent* Parent = OutputComponent->GetParent(); Parent; Parent = Parent->GetParent())
+		{
+			if (Parent == this)
+			{
+				OutputComponent->SetZOrder(++NewZOrder);
+			}
+		}
 	}
 }
 #endif // WITH_EDITOR
diff --git a/Engine/Plugins/VirtualProduction/DMX/DMXPixelMapping/Source/DMXPixelMappingRuntime/Public/Components/DMXPixelMappingOutputComponent.h b/Engine/Plugins/VirtualProduction/DMX/DMXPixelMapping/Source/DMXPixelMappingRuntime/Public/Components/DMXPixelMappingOutputComponent.h
--- a/Engine/Plugins/VirtualProduction/DMX/DMXPixelMapping/Source/DMXPixelMappingRuntime/Public/Components/DMXPixelMappingOutputComponent.h
+++ b/Engine/Plugins/VirtualProduction/DMX/DMXPixelMapping/Source/DMXPixelMappingRuntime/Public/Components/DMXPixelMappingOutputComponent.h
@@ -123,7 +123,7 @@ public:
 
 #if WITH_EDITOR
 	/** Makes the component the highest ZOrdered of components in the component rectangle, updates childs if needed */
-	void MakeHighestZOrderInComponentRect();
+	void ZOrderTopmost();
 #endif // WITH_EDITOR
 
 public:
