DMX - Fix Modulator instances don't update on property changes, show in the wrong category

diff --git a/Engine/Plugins/VirtualProduction/DMX/DMXPixelMapping/Source/DMXPixelMappingEditor/Private/Customizations/DMXPixelMappingDetailCustomization_FixtureGroupItem.cpp b/Engine/Plugins/VirtualProduction/DMX/DMXPixelMapping/Source/DMXPixelMappingEditor/Private/Customizations/DMXPixelMappingDetailCustomization_FixtureGroupItem.cpp
--- a/Engine/Plugins/VirtualProduction/DMX/DMXPixelMapping/Source/DMXPixelMappingEditor/Private/Customizations/DMXPixelMappingDetailCustomization_FixtureGroupItem.cpp
+++ b/Engine/Plugins/VirtualProduction/DMX/DMXPixelMapping/Source/DMXPixelMappingEditor/Private/Customizations/DMXPixelMappingDetailCustomization_FixtureGroupItem.cpp
@@ -2,38 +2,22 @@
 
 #include "Customizations/DMXPixelMappingDetailCustomization_FixtureGroupItem.h"
 
-#include "DMXPixelMappingTypes.h"
 #include "Components/DMXPixelMappingFixtureGroupItemComponent.h"
-#include "DetailsViewArgs.h"
-#include "Modulators/DMXModulator.h"
-
+#include "DetailCategoryBuilder.h"
 #include "DetailLayoutBuilder.h"
-#include "DetailWidgetRow.h"
-#include "IPropertyUtilities.h"
-#include "IDetailsView.h"
-#include "Modules/ModuleManager.h"
-#include "Widgets/Views/SListView.h"
 
-#define LOCTEXT_NAMESPACE "FixtureGroupItem"
 
+FDMXPixelMappingDetailCustomization_FixtureGroupItem::FDMXPixelMappingDetailCustomization_FixtureGroupItem(const TWeakPtr<FDMXPixelMappingToolkit> InToolkitWeakPtr)
+	: ToolkitWeakPtr(InToolkitWeakPtr)
+{}
 
-void FDMXPixelMappingDetailCustomization_FixtureGroupItem::CustomizeDetails(IDetailLayoutBuilder& InDetailLayout)
+TSharedRef<IDetailCustomization> FDMXPixelMappingDetailCustomization_FixtureGroupItem::MakeInstance(const TWeakPtr<FDMXPixelMappingToolkit> InToolkitWeakPtr)
 {
-	PropertyUtilities = InDetailLayout.GetPropertyUtilities();
-
-	// Get editing object
-	TArray<TWeakObjectPtr<UObject>> Objects;
-	InDetailLayout.GetObjectsBeingCustomized(Objects);
-
-	FixtureGroupItemComponents.Empty();
-	for (TWeakObjectPtr<UObject> SelectedObject : Objects)
-	{
-		FixtureGroupItemComponents.Add(Cast<UDMXPixelMappingFixtureGroupItemComponent>(SelectedObject));
-	}
-
-	// Get editing categories
-	IDetailCategoryBuilder& OutputSettingsCategory = InDetailLayout.EditCategory("Output Settings", FText::GetEmpty(), ECategoryPriority::Important);
+	return MakeShared<FDMXPixelMappingDetailCustomization_FixtureGroupItem>(InToolkitWeakPtr);
+}
 
+void FDMXPixelMappingDetailCustomization_FixtureGroupItem::CustomizeDetails(IDetailLayoutBuilder& InDetailLayout)
+{
 	// Hide absolute postition property handles
 	TSharedPtr<IPropertyHandle> PositionXPropertyHandle = InDetailLayout.GetProperty(UDMXPixelMappingOutputComponent::GetPositionXPropertyName());
 	InDetailLayout.HideProperty(PositionXPropertyHandle);
@@ -44,8 +28,6 @@ void FDMXPixelMappingDetailCustomization_FixtureGroupItem::CustomizeDetails(IDet
 	TSharedPtr<IPropertyHandle> SizeYPropertyHandle = InDetailLayout.GetProperty(UDMXPixelMappingOutputComponent::GetSizeYPropertyName());
 	InDetailLayout.HideProperty(SizeXPropertyHandle);
 
-	CreateModulatorDetails(InDetailLayout);
-
 	// Sort categories
 	InDetailLayout.SortCategories([](const TMap<FName, IDetailCategoryBuilder*>& CategoryMap)
 		{
@@ -113,102 +95,3 @@ void FDMXPixelMappingDetailCustomization_FixtureGroupItem::CustomizeDetails(IDet
 			}
 		});
 }
-
-void FDMXPixelMappingDetailCustomization_FixtureGroupItem::CreateModulatorDetails(IDetailLayoutBuilder& InDetailLayout)
-{
-	IDetailCategoryBuilder& ModualtorsCategory = InDetailLayout.EditCategory("Modulators", LOCTEXT("DMXModulatorsCategory", "Modulators"), ECategoryPriority::Important);
-
-	TSharedPtr<IPropertyHandle> ModulatorClassesHandle = InDetailLayout.GetProperty(GET_MEMBER_NAME_CHECKED(UDMXPixelMappingFixtureGroupItemComponent, ModulatorClasses), UDMXPixelMappingFixtureGroupItemComponent::StaticClass());
-	ModulatorClassesHandle->SetOnPropertyValueChanged(FSimpleDelegate::CreateSP(this, &FDMXPixelMappingDetailCustomization_FixtureGroupItem::RequestRefresh));
-	ModulatorClassesHandle->SetOnChildPropertyValueChanged(FSimpleDelegate::CreateSP(this, &FDMXPixelMappingDetailCustomization_FixtureGroupItem::RequestRefresh));
-
-	ModualtorsCategory.AddProperty(ModulatorClassesHandle);
-
-	TSharedPtr<IPropertyHandle> ModulatorsHandle = InDetailLayout.GetProperty(GET_MEMBER_NAME_CHECKED(UDMXPixelMappingFixtureGroupItemComponent, Modulators), UDMXPixelMappingFixtureGroupItemComponent::StaticClass());
-	InDetailLayout.HideProperty(ModulatorsHandle);
-
-	// Create detail views for the modulators
-	TArray<TWeakObjectPtr<UObject>> CustomizedObjects;
-	InDetailLayout.GetObjectsBeingCustomized(CustomizedObjects);
-	if (CustomizedObjects.Num() > 0)
-	{
-		if (UDMXPixelMappingFixtureGroupItemComponent* FirstGroupItemComponent = Cast<UDMXPixelMappingFixtureGroupItemComponent>(CustomizedObjects[0].Get()))
-		{
-			for (int32 IndexModulator = 0; IndexModulator < FirstGroupItemComponent->Modulators.Num(); IndexModulator++)
-			{
-				TArray<UObject*> ModulatorsToEdit;
-				if (CustomizedObjects.Num() > 1)
-				{
-					UClass* ModulatorClass = FirstGroupItemComponent->Modulators[IndexModulator]->GetClass();
-
-					for (const TWeakObjectPtr<UObject>& CustomizedObject : CustomizedObjects)
-					{
-						if (UDMXPixelMappingFixtureGroupItemComponent* GroupItemComponent = Cast<UDMXPixelMappingFixtureGroupItemComponent>(CustomizedObject.Get()))
-						{
-							const bool bMultiEditableModulator =
-								GroupItemComponent->Modulators.IsValidIndex(IndexModulator) &&
-								GroupItemComponent->Modulators[IndexModulator] &&
-								GroupItemComponent->Modulators[IndexModulator]->GetClass() == ModulatorClass;
-
-							if (bMultiEditableModulator)
-							{
-								ModulatorsToEdit.Add(GroupItemComponent->Modulators[IndexModulator]);
-							}
-							else
-							{
-								// Don't allow multi edit if not all modulators are of same class
-								ModulatorsToEdit.Reset();
-							}
-						}
-					}
-				}
-				else if (UDMXModulator* ModulatorOfFirstGroupItem = FirstGroupItemComponent->Modulators[IndexModulator])
-				{
-					ModulatorsToEdit.Add(ModulatorOfFirstGroupItem);
-				}
-
-
-				if (ModulatorsToEdit.Num() > 0)
-				{
-					FPropertyEditorModule& PropertyEditorModule = FModuleManager::GetModuleChecked<FPropertyEditorModule>("PropertyEditor");
-
-					FDetailsViewArgs DetailsViewArgs;
-					DetailsViewArgs.bUpdatesFromSelection = false;
-					DetailsViewArgs.bLockable = true;
-					DetailsViewArgs.bAllowSearch = false;
-					DetailsViewArgs.NameAreaSettings = FDetailsViewArgs::HideNameArea;
-					DetailsViewArgs.bHideSelectionTip = false;
-					DetailsViewArgs.DefaultsOnlyVisibility = EEditDefaultsOnlyNodeVisibility::Automatic;
-
-					TSharedRef<IDetailsView> DetailsView = PropertyEditorModule.CreateDetailView(DetailsViewArgs);
-					DetailsView->SetObjects(ModulatorsToEdit);
-
-					ModualtorsCategory.AddCustomRow(FText::GetEmpty())
-						.WholeRowContent()
-						[
-							DetailsView
-						];
-				}
-				else
-				{
-					ModualtorsCategory.AddCustomRow(FText::GetEmpty())
-						.WholeRowContent()
-						[
-							SNew(STextBlock)
-							.Text(LOCTEXT("ModulatorMultipleValues", "Multiple Values"))
-							.Font(IDetailLayoutBuilder::GetDetailFont())
-						];
-
-					break;
-				}
-			}
-		}
-	}
-}
-
-void FDMXPixelMappingDetailCustomization_FixtureGroupItem::RequestRefresh()
-{
-	PropertyUtilities->ForceRefresh();
-}
-
-#undef LOCTEXT_NAMESPACE
diff --git a/Engine/Plugins/VirtualProduction/DMX/DMXPixelMapping/Source/DMXPixelMappingEditor/Private/Customizations/DMXPixelMappingDetailCustomization_FixtureGroupItem.h b/Engine/Plugins/VirtualProduction/DMX/DMXPixelMapping/Source/DMXPixelMappingEditor/Private/Customizations/DMXPixelMappingDetailCustomization_FixtureGroupItem.h
--- a/Engine/Plugins/VirtualProduction/DMX/DMXPixelMapping/Source/DMXPixelMappingEditor/Private/Customizations/DMXPixelMappingDetailCustomization_FixtureGroupItem.h
+++ b/Engine/Plugins/VirtualProduction/DMX/DMXPixelMapping/Source/DMXPixelMappingEditor/Private/Customizations/DMXPixelMappingDetailCustomization_FixtureGroupItem.h
@@ -3,64 +3,24 @@
 #pragma once
 
 #include "IDetailCustomization.h"
-#include "Layout/Visibility.h"
-#include "Misc/Attribute.h"
-#include "UObject/WeakObjectPtr.h"
-
-
-class IPropertyUtilities;
 
 class FDMXPixelMappingToolkit;
 class IDetailLayoutBuilder;
-class IPropertyHandle;
-class UDMXPixelMappingFixtureGroupItemComponent;
-enum class EDMXColorMode : uint8;
-
-class IPropertyUitilities;
-template <typename ItemType> class SListView;
 
 
 class FDMXPixelMappingDetailCustomization_FixtureGroupItem
 	: public IDetailCustomization
 {
-private:
-	struct FFunctionAttribute
-	{
-		TSharedPtr<IPropertyHandle> Handle;
-		TSharedPtr<IPropertyHandle> ExposeHandle;
-		TSharedPtr<IPropertyHandle> InvertHandle;
-		TAttribute<EVisibility> Visibility;
-	};
-
-
 public:
-	static TSharedRef<IDetailCustomization> MakeInstance(const TWeakPtr<FDMXPixelMappingToolkit> InToolkitWeakPtr)
-	{
-		return MakeShared<FDMXPixelMappingDetailCustomization_FixtureGroupItem>(InToolkitWeakPtr);
-	}
-
-	/** Constructor */
-	FDMXPixelMappingDetailCustomization_FixtureGroupItem(const TWeakPtr<FDMXPixelMappingToolkit> InToolkitWeakPtr)
-		: ToolkitWeakPtr(InToolkitWeakPtr)
-	{}
+	static TSharedRef<IDetailCustomization> MakeInstance(const TWeakPtr<FDMXPixelMappingToolkit> InToolkitWeakPtr);
 
+	FDMXPixelMappingDetailCustomization_FixtureGroupItem(const TWeakPtr<FDMXPixelMappingToolkit> InToolkitWeakPtr);
+	
 	//~ IDetailCustomization interface begin
 	virtual void CustomizeDetails(IDetailLayoutBuilder& DetailLayout) override;
 	//~ IPropertyTypeCustomization interface end
 
 private:
-	/** Creates Details for the Output Modulators */
-	void CreateModulatorDetails(IDetailLayoutBuilder& InDetailLayout);
-
-	/** Requests to refresh the details view */
-	void RequestRefresh();
-
-	/** Fixture Group Item Components that are being edited */
-	TArray<TWeakObjectPtr<UDMXPixelMappingFixtureGroupItemComponent>> FixtureGroupItemComponents;
-
-	/** Property utilities for this customization */
-	TSharedPtr<IPropertyUtilities> PropertyUtilities;
-
 	/** Weak reference to the DMX editor */
 	TWeakPtr<FDMXPixelMappingToolkit> ToolkitWeakPtr;
 };
diff --git a/Engine/Plugins/VirtualProduction/DMX/DMXPixelMapping/Source/DMXPixelMappingEditor/Private/Customizations/DMXPixelMappingDetailCustomization_Matrix.cpp b/Engine/Plugins/VirtualProduction/DMX/DMXPixelMapping/Source/DMXPixelMappingEditor/Private/Customizations/DMXPixelMappingDetailCustomization_Matrix.cpp
--- a/Engine/Plugins/VirtualProduction/DMX/DMXPixelMapping/Source/DMXPixelMappingEditor/Private/Customizations/DMXPixelMappingDetailCustomization_Matrix.cpp
+++ b/Engine/Plugins/VirtualProduction/DMX/DMXPixelMapping/Source/DMXPixelMappingEditor/Private/Customizations/DMXPixelMappingDetailCustomization_Matrix.cpp
@@ -2,45 +2,25 @@
 
 #include "Customizations/DMXPixelMappingDetailCustomization_Matrix.h"
 
-#include "DMXPixelMappingTypes.h"
 #include "Components/DMXPixelMappingMatrixComponent.h"
-#include "Modulators/DMXModulator.h"
-
+#include "DetailCategoryBuilder.h"
 #include "DetailLayoutBuilder.h"
-#include "DetailsViewArgs.h"
-#include "DetailWidgetRow.h"
-#include "IPropertyUtilities.h"
-#include "IDetailsView.h"
-#include "Framework/Views/TableViewMetadata.h"
-#include "Modules/ModuleManager.h"
-#include "Widgets/Views/SListView.h"
 
 
-#define LOCTEXT_NAMESPACE "DMXPixelMappingDetailCustomization_Matrix"
+FDMXPixelMappingDetailCustomization_Matrix::FDMXPixelMappingDetailCustomization_Matrix(const TWeakPtr<FDMXPixelMappingToolkit> InToolkitWeakPtr)
+	: ToolkitWeakPtr(InToolkitWeakPtr)
+{}
 
-void FDMXPixelMappingDetailCustomization_Matrix::CustomizeDetails(IDetailLayoutBuilder& InDetailLayout)
+TSharedRef<IDetailCustomization> FDMXPixelMappingDetailCustomization_Matrix::MakeInstance(const TWeakPtr<FDMXPixelMappingToolkit> InToolkitWeakPtr)
 {
-	PropertyUtilities = InDetailLayout.GetPropertyUtilities();
-
-	// Get editing UObject
-	TArray<TWeakObjectPtr<UObject>> OuterObjects;
-	InDetailLayout.GetObjectsBeingCustomized(OuterObjects);
-
-	MatrixComponents.Empty();
-
-	for (TWeakObjectPtr<UObject> SelectedObject : OuterObjects)
-	{
-		MatrixComponents.Add(Cast<UDMXPixelMappingMatrixComponent>(SelectedObject));
-	}
-
-	// Get editing categories
-	IDetailCategoryBuilder& OutputSettingsCategory = InDetailLayout.EditCategory("Output Settings", FText::GetEmpty(), ECategoryPriority::Important);
+	return MakeShared<FDMXPixelMappingDetailCustomization_Matrix>(InToolkitWeakPtr);
+}
 
+void FDMXPixelMappingDetailCustomization_Matrix::CustomizeDetails(IDetailLayoutBuilder& InDetailLayout)
+{
 	// Hide the Layout Script property (shown in its own panel, see SDMXPixelMappingLayoutView)
 	InDetailLayout.HideProperty(GET_MEMBER_NAME_CHECKED(UDMXPixelMappingMatrixComponent, LayoutScript));
 
-	CreateModulatorDetails(InDetailLayout);
-
 	// Sort categories
 	InDetailLayout.SortCategories([](const TMap<FName, IDetailCategoryBuilder*>& CategoryMap)
 		{
@@ -108,101 +88,3 @@ void FDMXPixelMappingDetailCustomization_Matrix::CustomizeDetails(IDetailLayoutB
 			}
 		});
 }
-
-void FDMXPixelMappingDetailCustomization_Matrix::CreateModulatorDetails(IDetailLayoutBuilder& InDetailLayout)
-{
-	IDetailCategoryBuilder& ModualtorsCategory = InDetailLayout.EditCategory("Modulators", LOCTEXT("DMXModulatorsCategory", "Modulators"), ECategoryPriority::Important);
-
-	TSharedPtr<IPropertyHandle> ModulatorClassesHandle = InDetailLayout.GetProperty(GET_MEMBER_NAME_CHECKED(UDMXPixelMappingMatrixComponent, ModulatorClasses), UDMXPixelMappingMatrixComponent::StaticClass());
-	ModulatorClassesHandle->SetOnPropertyValueChanged(FSimpleDelegate::CreateSP(this, &FDMXPixelMappingDetailCustomization_Matrix::RequestRefresh));
-	ModulatorClassesHandle->SetOnChildPropertyValueChanged(FSimpleDelegate::CreateSP(this, &FDMXPixelMappingDetailCustomization_Matrix::RequestRefresh));
-
-	ModualtorsCategory.AddProperty(ModulatorClassesHandle);
-
-	TSharedPtr<IPropertyHandle> ModulatorsHandle = InDetailLayout.GetProperty(GET_MEMBER_NAME_CHECKED(UDMXPixelMappingMatrixComponent, Modulators), UDMXPixelMappingMatrixComponent::StaticClass());
-	InDetailLayout.HideProperty(ModulatorsHandle);
-
-	// Create detail views for the modulators
-	TArray<TWeakObjectPtr<UObject>> CustomizedObjects;
-	InDetailLayout.GetObjectsBeingCustomized(CustomizedObjects);
-	if (CustomizedObjects.Num() > 0)
-	{
-		if (UDMXPixelMappingMatrixComponent* FirstMatrixComponent = Cast<UDMXPixelMappingMatrixComponent>(CustomizedObjects[0].Get()))
-		{
-			for (int32 IndexModulator = 0; IndexModulator < FirstMatrixComponent->Modulators.Num(); IndexModulator++)
-			{
-				TArray<UObject*> ModulatorsToEdit;
-				if (CustomizedObjects.Num() > 1)
-				{
-					UClass* ModulatorClass = FirstMatrixComponent->Modulators[IndexModulator]->GetClass();
-
-					for (const TWeakObjectPtr<UObject>& CustomizedObject : CustomizedObjects)
-					{
-						if (UDMXPixelMappingMatrixComponent* MatrixComponent = Cast<UDMXPixelMappingMatrixComponent>(CustomizedObject.Get()))
-						{
-							const bool bMultiEditableModulator =
-								MatrixComponent->Modulators.IsValidIndex(IndexModulator) &&
-								MatrixComponent->Modulators[IndexModulator] &&
-								MatrixComponent->Modulators[IndexModulator]->GetClass() == ModulatorClass;
-
-							if (bMultiEditableModulator)
-							{
-								ModulatorsToEdit.Add(MatrixComponent->Modulators[IndexModulator]);
-							}
-							else
-							{
-								// Don't allow multi edit if not all modulators are of same class
-								ModulatorsToEdit.Reset();
-							}
-						}
-					}
-				}
-				else if (UDMXModulator* ModulatorOfFirstMatrix = FirstMatrixComponent->Modulators[IndexModulator])
-				{
-					ModulatorsToEdit.Add(ModulatorOfFirstMatrix);
-				}
-
-				if (ModulatorsToEdit.Num() > 0)
-				{
-					FPropertyEditorModule& PropertyEditorModule = FModuleManager::GetModuleChecked<FPropertyEditorModule>("PropertyEditor");
-
-					FDetailsViewArgs DetailsViewArgs;
-					DetailsViewArgs.bUpdatesFromSelection = false;
-					DetailsViewArgs.bLockable = true;
-					DetailsViewArgs.bAllowSearch = false;
-					DetailsViewArgs.bHideSelectionTip = false;
-					DetailsViewArgs.NameAreaSettings = FDetailsViewArgs::HideNameArea;
-					DetailsViewArgs.DefaultsOnlyVisibility = EEditDefaultsOnlyNodeVisibility::Automatic;
-
-					TSharedRef<IDetailsView> DetailsView = PropertyEditorModule.CreateDetailView(DetailsViewArgs);
-					DetailsView->SetObjects(ModulatorsToEdit);
-
-					ModualtorsCategory.AddCustomRow(FText::GetEmpty())
-						.WholeRowContent()
-						[
-							DetailsView
-						];
-				}
-				else
-				{
-					ModualtorsCategory.AddCustomRow(FText::GetEmpty())
-						.WholeRowContent()
-						[
-							SNew(STextBlock)
-							.Text(LOCTEXT("ModulatorMultipleValues", "Multiple Values"))
-							.Font(IDetailLayoutBuilder::GetDetailFont())
-						];
-
-					break;
-				}
-			}
-		}
-	}
-}
-
-void FDMXPixelMappingDetailCustomization_Matrix::RequestRefresh()
-{
-	PropertyUtilities->RequestRefresh();
-}
-
-#undef LOCTEXT_NAMESPACE
diff --git a/Engine/Plugins/VirtualProduction/DMX/DMXPixelMapping/Source/DMXPixelMappingEditor/Private/Customizations/DMXPixelMappingDetailCustomization_Matrix.h b/Engine/Plugins/VirtualProduction/DMX/DMXPixelMapping/Source/DMXPixelMappingEditor/Private/Customizations/DMXPixelMappingDetailCustomization_Matrix.h
--- a/Engine/Plugins/VirtualProduction/DMX/DMXPixelMapping/Source/DMXPixelMappingEditor/Private/Customizations/DMXPixelMappingDetailCustomization_Matrix.h
+++ b/Engine/Plugins/VirtualProduction/DMX/DMXPixelMapping/Source/DMXPixelMappingEditor/Private/Customizations/DMXPixelMappingDetailCustomization_Matrix.h
@@ -3,62 +3,23 @@
 #pragma once
 
 #include "IDetailCustomization.h"
-#include "Misc/Attribute.h"
-#include "Layout/Visibility.h"
-#include "UObject/WeakObjectPtr.h"
 
 class FDMXPixelMappingToolkit;
-struct FDMXCellAttributeGroup;
-class UDMXPixelMappingMatrixComponent;
-class UDMXEntityFixturePatch;
-
-enum class EDMXColorMode : uint8;
-class ITableRow;
-class IPropertyHandle;
-class IPropertyUtilities;
-class STableViewBase;
-template <typename ItemType> class SListView;
+class IDetailLayoutBuilder;
 
 
 class FDMXPixelMappingDetailCustomization_Matrix
 	: public IDetailCustomization
 {
-private:
-	struct FDMXCellAttributeGroup
-	{
-		TSharedPtr<IPropertyHandle> Handle;
-		TSharedPtr<IPropertyHandle> ExposeHandle;
-		TSharedPtr<IPropertyHandle> InvertHandle;
-		TAttribute<EVisibility> Visibility;
-	};
-
 public:
-	static TSharedRef<IDetailCustomization> MakeInstance(const TWeakPtr<FDMXPixelMappingToolkit> InToolkitWeakPtr)
-	{
-		return MakeShared<FDMXPixelMappingDetailCustomization_Matrix>(InToolkitWeakPtr);
-	}
+	FDMXPixelMappingDetailCustomization_Matrix(const TWeakPtr<FDMXPixelMappingToolkit> InToolkitWeakPtr);
 
-	/** Constructor */
-	FDMXPixelMappingDetailCustomization_Matrix(const TWeakPtr<FDMXPixelMappingToolkit> InToolkitWeakPtr)
-		: ToolkitWeakPtr(InToolkitWeakPtr)
-	{}
+	static TSharedRef<IDetailCustomization> MakeInstance(const TWeakPtr<FDMXPixelMappingToolkit> InToolkitWeakPtr);
 
 	/** IDetailCustomization interface */
 	virtual void CustomizeDetails(IDetailLayoutBuilder& DetailLayout) override;
 
 private:	
-	/** Creates Details for the Output Modulators */
-	void CreateModulatorDetails(IDetailLayoutBuilder& InDetailLayout);
-
-	/** Requests to refresh the details view */
-	void RequestRefresh();
-
-	/** Matrix Components that are being edited */
-	TArray<TWeakObjectPtr<UDMXPixelMappingMatrixComponent>> MatrixComponents;
-
-	/** Property utilities for this customization */
-	TSharedPtr<IPropertyUtilities> PropertyUtilities;
-
 	/** Weak reference to the DMX editor */
 	TWeakPtr<FDMXPixelMappingToolkit> ToolkitWeakPtr;
 };
diff --git a/Engine/Plugins/VirtualProduction/DMX/DMXPixelMapping/Source/DMXPixelMappingEditor/Private/Customizations/DMXPixelMappingDetailCustomization_OutputDMX.cpp b/Engine/Plugins/VirtualProduction/DMX/DMXPixelMapping/Source/DMXPixelMappingEditor/Private/Customizations/DMXPixelMappingDetailCustomization_OutputDMX.cpp
new file mode 100644
--- /dev/null
+++ b/Engine/Plugins/VirtualProduction/DMX/DMXPixelMapping/Source/DMXPixelMappingEditor/Private/Customizations/DMXPixelMappingDetailCustomization_OutputDMX.cpp
@@ -0,0 +1,123 @@
+// Copyright Epic Games, Inc. All Rights Reserved.
+
+#include "Customizations/DMXPixelMappingDetailCustomization_OutputDMX.h"
+
+#include "Components/DMXPixelMappingOutputDMXComponent.h"
+#include "DetailsViewArgs.h"
+#include "DetailLayoutBuilder.h"
+#include "DetailWidgetRow.h"
+#include "IPropertyUtilities.h"
+#include "IDetailsView.h"
+#include "Modulators/DMXModulator.h"
+#include "Modules/ModuleManager.h"
+#include "Widgets/Text/STextBlock.h"
+
+
+#define LOCTEXT_NAMESPACE "DMXPixelMappingDetailCustomization_OutputDMX"
+
+void FDMXPixelMappingDetailCustomization_OutputDMX::CustomizeDetails(IDetailLayoutBuilder& InDetailLayout)
+{
+	PropertyUtilities = InDetailLayout.GetPropertyUtilities();
+
+	CreateModulatorDetails(InDetailLayout);
+}
+
+void FDMXPixelMappingDetailCustomization_OutputDMX::CreateModulatorDetails(IDetailLayoutBuilder& InDetailLayout)
+{
+	IDetailCategoryBuilder& ModualtorsCategory = InDetailLayout.EditCategory("Output Modulators", LOCTEXT("DMXModulatorsCategory", "Output Modulators"), ECategoryPriority::Important);
+
+	const TSharedPtr<IPropertyHandle> ModulatorClassesHandle = InDetailLayout.GetProperty(GET_MEMBER_NAME_CHECKED(UDMXPixelMappingOutputDMXComponent, ModulatorClasses));
+	ModulatorClassesHandle->SetOnPropertyValueChanged(FSimpleDelegate::CreateSP(this, &FDMXPixelMappingDetailCustomization_OutputDMX::ForceRefresh));
+	ModulatorClassesHandle->SetOnChildPropertyValueChanged(FSimpleDelegate::CreateSP(this, &FDMXPixelMappingDetailCustomization_OutputDMX::ForceRefresh));
+
+	ModualtorsCategory.AddProperty(ModulatorClassesHandle);
+
+	const TSharedPtr<IPropertyHandle> ModulatorsHandle = InDetailLayout.GetProperty(GET_MEMBER_NAME_CHECKED(UDMXPixelMappingOutputDMXComponent, Modulators));
+	InDetailLayout.HideProperty(ModulatorsHandle);
+
+	// Create detail views for the modulators
+	TArray<TWeakObjectPtr<UObject>> CustomizedObjects;
+	InDetailLayout.GetObjectsBeingCustomized(CustomizedObjects);
+	if (CustomizedObjects.Num() > 0)
+	{
+		if (UDMXPixelMappingOutputDMXComponent* FirstOutputDMXComponent = Cast<UDMXPixelMappingOutputDMXComponent>(CustomizedObjects[0].Get()))
+		{
+			for (int32 IndexModulator = 0; IndexModulator < FirstOutputDMXComponent->Modulators.Num(); IndexModulator++)
+			{
+				TArray<UObject*> ModulatorsToEdit;
+				if (CustomizedObjects.Num() > 1)
+				{
+					UClass* ModulatorClass = FirstOutputDMXComponent->Modulators[IndexModulator]->GetClass();
+
+					for (const TWeakObjectPtr<UObject>& CustomizedObject : CustomizedObjects)
+					{
+						if (UDMXPixelMappingOutputDMXComponent* OutputDMXComponent = Cast<UDMXPixelMappingOutputDMXComponent>(CustomizedObject.Get()))
+						{
+							const bool bMultiEditableModulator =
+								OutputDMXComponent->Modulators.IsValidIndex(IndexModulator) &&
+								OutputDMXComponent->Modulators[IndexModulator] &&
+								OutputDMXComponent->Modulators[IndexModulator]->GetClass() == ModulatorClass;
+
+							if (bMultiEditableModulator)
+							{
+								ModulatorsToEdit.Add(OutputDMXComponent->Modulators[IndexModulator]);
+							}
+							else
+							{
+								// Don't allow multi edit if not all modulators are of same class
+								ModulatorsToEdit.Reset();
+							}
+						}
+					}
+				}
+				else if (UDMXModulator* ModulatorOfFirstGroupItem = FirstOutputDMXComponent->Modulators[IndexModulator])
+				{
+					ModulatorsToEdit.Add(ModulatorOfFirstGroupItem);
+				}
+
+
+				if (ModulatorsToEdit.Num() > 0)
+				{
+					FPropertyEditorModule& PropertyEditorModule = FModuleManager::GetModuleChecked<FPropertyEditorModule>("PropertyEditor");
+
+					FDetailsViewArgs DetailsViewArgs;
+					DetailsViewArgs.bUpdatesFromSelection = false;
+					DetailsViewArgs.bLockable = true;
+					DetailsViewArgs.bAllowSearch = false;
+					DetailsViewArgs.NameAreaSettings = FDetailsViewArgs::HideNameArea;
+					DetailsViewArgs.bHideSelectionTip = false;
+					DetailsViewArgs.DefaultsOnlyVisibility = EEditDefaultsOnlyNodeVisibility::Automatic;
+
+					TSharedRef<IDetailsView> DetailsView = PropertyEditorModule.CreateDetailView(DetailsViewArgs);
+					DetailsView->SetObjects(ModulatorsToEdit);
+
+					ModualtorsCategory.AddCustomRow(FText::GetEmpty())
+						.WholeRowContent()
+						[
+							DetailsView
+						];
+				}
+				else
+				{
+					ModualtorsCategory.AddCustomRow(FText::GetEmpty())
+						.WholeRowContent()
+						[
+							SNew(STextBlock)
+							.Text(LOCTEXT("ModulatorMultipleValues", "Multiple Values"))
+							.Font(IDetailLayoutBuilder::GetDetailFont())
+						];
+
+					break;
+				}
+			}
+		}
+	}
+}
+
+void FDMXPixelMappingDetailCustomization_OutputDMX::ForceRefresh()
+{	
+	// For some reason IPropertyUtilities::RequestRefresh doesn't recognize changed modulators
+	PropertyUtilities->ForceRefresh();
+}
+
+#undef LOCTEXT_NAMESPACE
diff --git a/Engine/Plugins/VirtualProduction/DMX/DMXPixelMapping/Source/DMXPixelMappingEditor/Private/Customizations/DMXPixelMappingDetailCustomization_OutputDMX.h b/Engine/Plugins/VirtualProduction/DMX/DMXPixelMapping/Source/DMXPixelMappingEditor/Private/Customizations/DMXPixelMappingDetailCustomization_OutputDMX.h
new file mode 100644
--- /dev/null
+++ b/Engine/Plugins/VirtualProduction/DMX/DMXPixelMapping/Source/DMXPixelMappingEditor/Private/Customizations/DMXPixelMappingDetailCustomization_OutputDMX.h
@@ -0,0 +1,42 @@
+// Copyright Epic Games, Inc. All Rights Reserved.
+
+#pragma once
+
+#include "IDetailCustomization.h"
+
+class IPropertyUtilities;
+class FDMXPixelMappingToolkit;
+class IDetailLayoutBuilder;
+
+
+class FDMXPixelMappingDetailCustomization_OutputDMX
+	: public IDetailCustomization
+{
+public:
+	static TSharedRef<IDetailCustomization> MakeInstance(const TWeakPtr<FDMXPixelMappingToolkit> InToolkitWeakPtr)
+	{
+		return MakeShared<FDMXPixelMappingDetailCustomization_OutputDMX>(InToolkitWeakPtr);
+	}
+
+	/** Constructor */
+	FDMXPixelMappingDetailCustomization_OutputDMX(const TWeakPtr<FDMXPixelMappingToolkit> InToolkitWeakPtr)
+		: ToolkitWeakPtr(InToolkitWeakPtr)
+	{}
+
+	//~ IDetailCustomization interface begin
+	virtual void CustomizeDetails(IDetailLayoutBuilder& DetailLayout) override;
+	//~ IPropertyTypeCustomization interface end
+
+private:
+	/** Creates Details for the Output Modulators */
+	void CreateModulatorDetails(IDetailLayoutBuilder& InDetailLayout);
+	
+	/** Forces the details to refresh */
+	void ForceRefresh();
+
+	/** Property utilities for this customization */
+	TSharedPtr<IPropertyUtilities> PropertyUtilities;
+
+	/** Weak reference to the DMX editor */
+	TWeakPtr<FDMXPixelMappingToolkit> ToolkitWeakPtr;
+};
diff --git a/Engine/Plugins/VirtualProduction/DMX/DMXPixelMapping/Source/DMXPixelMappingEditor/Private/Views/SDMXPixelMappingDetailsView.cpp b/Engine/Plugins/VirtualProduction/DMX/DMXPixelMapping/Source/DMXPixelMappingEditor/Private/Views/SDMXPixelMappingDetailsView.cpp
--- a/Engine/Plugins/VirtualProduction/DMX/DMXPixelMapping/Source/DMXPixelMappingEditor/Private/Views/SDMXPixelMappingDetailsView.cpp
+++ b/Engine/Plugins/VirtualProduction/DMX/DMXPixelMapping/Source/DMXPixelMappingEditor/Private/Views/SDMXPixelMappingDetailsView.cpp
@@ -9,6 +9,7 @@
 #include "Components/DMXPixelMappingScreenComponent.h"
 #include "Components/DMXPixelMappingMatrixComponent.h"
 #include "Components/DMXPixelMappingMatrixCellComponent.h"
+#include "Components/DMXPixelMappingOutputDMXComponent.h"
 #include "Components/DMXPixelMappingRendererComponent.h"
 #include "Customizations/DMXPixelMappingColorSpaceDetails_RGBCMY.h"
 #include "Customizations/DMXPixelMappingDetailCustomization_FixtureGroup.h"
@@ -16,6 +17,7 @@
 #include "Customizations/DMXPixelMappingDetailCustomization_Screen.h"
 #include "Customizations/DMXPixelMappingDetailCustomization_Renderer.h"
 #include "Customizations/DMXPixelMappingDetailCustomization_Matrix.h"
+#include "Customizations/DMXPixelMappingDetailCustomization_OutputDMX.h"
 #include "DetailsViewArgs.h"
 #include "DMXPixelMappingComponentReference.h"
 #include "IDetailsView.h"
@@ -107,14 +109,17 @@ void SDMXPixelMappingDetailsView::RegisterCustomizations()
 	FOnGetDetailCustomizationInstance FixtureGroupCustomizationInstance = FOnGetDetailCustomizationInstance::CreateStatic(&FDMXPixelMappingDetailCustomization_FixtureGroup::MakeInstance, WeakToolkit);
 	DetailsView->RegisterInstancedCustomPropertyLayout(UDMXPixelMappingFixtureGroupComponent::StaticClass(), FixtureGroupCustomizationInstance);
 
+	FOnGetDetailCustomizationInstance OutputDMXCustomizationInstance = FOnGetDetailCustomizationInstance::CreateStatic(&FDMXPixelMappingDetailCustomization_OutputDMX::MakeInstance, WeakToolkit);
+	DetailsView->RegisterInstancedCustomPropertyLayout(UDMXPixelMappingOutputDMXComponent::StaticClass(), OutputDMXCustomizationInstance);
+
 	FOnGetDetailCustomizationInstance FixtureGroupItemCustomizationInstance = FOnGetDetailCustomizationInstance::CreateStatic(&FDMXPixelMappingDetailCustomization_FixtureGroupItem::MakeInstance, WeakToolkit);
 	DetailsView->RegisterInstancedCustomPropertyLayout(UDMXPixelMappingFixtureGroupItemComponent::StaticClass(), FixtureGroupItemCustomizationInstance);
-	
-	FOnGetDetailCustomizationInstance ScreenCustomizationInstance = FOnGetDetailCustomizationInstance::CreateStatic(&FDMXPixelMappingDetailCustomization_Screen::MakeInstance, WeakToolkit);
-	DetailsView->RegisterInstancedCustomPropertyLayout(UDMXPixelMappingScreenComponent::StaticClass(), ScreenCustomizationInstance);
 
 	FOnGetDetailCustomizationInstance MatrixCustomizationInstance = FOnGetDetailCustomizationInstance::CreateStatic(&FDMXPixelMappingDetailCustomization_Matrix::MakeInstance, WeakToolkit);
 	DetailsView->RegisterInstancedCustomPropertyLayout(UDMXPixelMappingMatrixComponent::StaticClass(), MatrixCustomizationInstance);
+	
+	FOnGetDetailCustomizationInstance ScreenCustomizationInstance = FOnGetDetailCustomizationInstance::CreateStatic(&FDMXPixelMappingDetailCustomization_Screen::MakeInstance, WeakToolkit);
+	DetailsView->RegisterInstancedCustomPropertyLayout(UDMXPixelMappingScreenComponent::StaticClass(), ScreenCustomizationInstance);
 
 	FOnGetDetailCustomizationInstance ColorSpaceCustomizationInstance_RGBCMY = FOnGetDetailCustomizationInstance::CreateStatic(&FDMXPixelMappingColorSpaceDetails_RGBCMY::MakeInstance);
 	DetailsView->RegisterInstancedCustomPropertyLayout(UDMXPixelMappingColorSpace_RGBCMY::StaticClass(), ColorSpaceCustomizationInstance_RGBCMY);
