DMX - Rename Texture Material param for pixel mapping blur material to "Texture"

diff --git a/Engine/Plugins/VirtualProduction/DMX/DMXPixelMapping/Source/DMXPixelMappingRenderer/Private/PreprocessApplyFilterMaterialProxy.cpp b/Engine/Plugins/VirtualProduction/DMX/DMXPixelMapping/Source/DMXPixelMappingRenderer/Private/PreprocessApplyFilterMaterialProxy.cpp
--- a/Engine/Plugins/VirtualProduction/DMX/DMXPixelMapping/Source/DMXPixelMappingRenderer/Private/PreprocessApplyFilterMaterialProxy.cpp
+++ b/Engine/Plugins/VirtualProduction/DMX/DMXPixelMapping/Source/DMXPixelMappingRenderer/Private/PreprocessApplyFilterMaterialProxy.cpp
@@ -97,7 +97,7 @@ namespace UE::DMXPixelMapping::Rendering::Preprocess::Private
 		if (MaterialInstanceDynamic && DownsampleRenderTargets.IsEmpty() && ensureMsgf(OutputRenderTarget, TEXT("Cannot apply pixel mapping filter material. Missing output render target.")))
 		{
 			// Render without downsampling
-			MaterialInstanceDynamic->SetTextureParameterValue(InPreprocessRenderer.GetInputTextureParameterName(), InputTexture);
+			MaterialInstanceDynamic->SetTextureParameterValue(InPreprocessRenderer.GetTextureParameterName(), InputTexture);
 			MaterialInstanceDynamic->SetScalarParameterValue(InPreprocessRenderer.GetBlurDistanceParameterName(), InPreprocessRenderer.GetBlurDistance());
 			DrawMaterialCanvas.DrawMaterialToRenderTarget(OutputRenderTarget, MaterialInstanceDynamic);
 		}
@@ -117,7 +117,7 @@ namespace UE::DMXPixelMapping::Rendering::Preprocess::Private
 				const bool bApplyFilterMaterial = MaterialInstanceDynamic && (bApplyFilterMaterialEachDownsamplePass || DownsampleRenderTarget == DownsampleRenderTargets.Last());
 				if (bApplyFilterMaterial)
 				{
-					MaterialInstanceDynamic->SetTextureParameterValue(InPreprocessRenderer.GetInputTextureParameterName(), SourceTexture);
+					MaterialInstanceDynamic->SetTextureParameterValue(InPreprocessRenderer.GetTextureParameterName(), SourceTexture);
 					MaterialInstanceDynamic->SetScalarParameterValue(InPreprocessRenderer.GetBlurDistanceParameterName(), InPreprocessRenderer.GetBlurDistance());
 					DrawMaterialCanvas.DrawMaterialToRenderTarget(DownsampleRenderTarget, MaterialInstanceDynamic);
 				}
diff --git a/Engine/Plugins/VirtualProduction/DMX/DMXPixelMapping/Source/DMXPixelMappingRenderer/Public/DMXPixelMappingPreprocessRenderer.h b/Engine/Plugins/VirtualProduction/DMX/DMXPixelMapping/Source/DMXPixelMappingRenderer/Public/DMXPixelMappingPreprocessRenderer.h
--- a/Engine/Plugins/VirtualProduction/DMX/DMXPixelMapping/Source/DMXPixelMappingRenderer/Public/DMXPixelMappingPreprocessRenderer.h
+++ b/Engine/Plugins/VirtualProduction/DMX/DMXPixelMapping/Source/DMXPixelMappingRenderer/Public/DMXPixelMappingPreprocessRenderer.h
@@ -99,10 +99,10 @@ public:
 	/** Returns the blur distance parameter */
 	float GetBlurDistance() const { return BlurDistance; }
 
-	/** Returns the input texture material parameter name */
-	FName GetInputTextureParameterName() const { return InputTextureParameterName; }
+	/** Returns the texture parameter name */
+	FName GetTextureParameterName() const { return TextureParameterName; }
 
-	/** Returns the blur distance material parameter name */
+	/** Returns the blur distance parameter name */
 	FName GetBlurDistanceParameterName() const { return BlurDistanceParamterName; }
 
 	// Property name getters
@@ -152,11 +152,11 @@ private:
 	UPROPERTY(EditAnywhere, Category = "Filtering", Meta = (UIMin = "0", UIMax = "1", SliderExponent = 2.0, AllowPrivateAccess = true))
 	float BlurDistance = .02f;
 
-	/** Parameter name of the filter material to receive the input texture */
+	/** The texture parameter name in the Filter Material */
 	UPROPERTY(EditAnywhere, AdvancedDisplay, Category = "Filtering", Meta = (AllowPrivateAccess = true))
-	FName InputTextureParameterName = "InputTexture";
+	FName TextureParameterName = "Texture";
 
-	/** The blur distance parameter name */
+	/** The blur distance parameter name in the Filter Material */
 	UPROPERTY(EditAnywhere, AdvancedDisplay, Category = "Filtering", Meta = (AllowPrivateAccess = true))
 	FName BlurDistanceParamterName = "BlurDistance";
 
