DMX - Fix various PixelMapping UX bugs

* Newly added components are now selected
* Rename works as expected
* Add DMX Library is now called Add Fixture Group
* Undo now works for 'add mapping'

diff --git a/Engine/Plugins/VirtualProduction/DMX/DMXPixelMapping/Source/DMXPixelMappingBlueprintGraph/Private/Widgets/SDMXPixelMappingComponentPin.h b/Engine/Plugins/VirtualProduction/DMX/DMXPixelMapping/Source/DMXPixelMappingBlueprintGraph/Private/Widgets/SDMXPixelMappingComponentPin.h
--- a/Engine/Plugins/VirtualProduction/DMX/DMXPixelMapping/Source/DMXPixelMappingBlueprintGraph/Private/Widgets/SDMXPixelMappingComponentPin.h
+++ b/Engine/Plugins/VirtualProduction/DMX/DMXPixelMapping/Source/DMXPixelMappingBlueprintGraph/Private/Widgets/SDMXPixelMappingComponentPin.h
@@ -9,7 +9,7 @@
 
 #include "DMXPixelMapping.h" // IWYU pragma: keep
 
-struct FComponentUserFriendlyNameTable;
+struct FComponentUserNameTable;
 
 /** 
  * Cusotom widget for Pixel Mapping component pin.
@@ -26,7 +26,7 @@ public:
 	{
 		DMXPixelMappingWeakPtr = InDMXPixelMapping;
 
-		NameTable = FComponentUserFriendlyNameTable(InDMXPixelMapping);
+		NameTable = FComponentUserNameTable(InDMXPixelMapping);
 
 		SGraphPin::Construct(SGraphPin::FArguments(), InGraphPinObj);
 	}
@@ -40,7 +40,7 @@ protected:
 	 */
 	virtual TSharedRef<SWidget>	GetDefaultValueWidget() override
 	{
-		NameTable.GetUserFriendlyNames(UserFriendlyNameList);
+		NameTable.GetUserNames(UserNameList);
 
 		// Preserve previous selection if possible
 		const FName SelectedComponentName = [&]()
@@ -63,24 +63,24 @@ protected:
 
 		SetNameToPin(SelectedComponentName);
 
-		// Show the user friendly name, instead of the component's FName
-		FString UserFriendlyName = NameTable.GetUserFriendlyName(SelectedComponentName);
-		const TSharedPtr<FString>* SelectedUserFriendlyNamePtr = UserFriendlyNameList.FindByPredicate([UserFriendlyName](const TSharedPtr<FString>& TestedUserFriendlyName) {
-			return TestedUserFriendlyName->Equals(UserFriendlyName);
+		// Show the user  name, instead of the component's FName
+		FString UserName = NameTable.GetUserName(SelectedComponentName);
+		const TSharedPtr<FString>* SelectedUserNamePtr = UserNameList.FindByPredicate([UserName](const TSharedPtr<FString>& TestedUserName) {
+			return TestedUserName->Equals(UserName);
 			});
-		const TSharedPtr<FString> SelectedUserFriendlyName = SelectedUserFriendlyNamePtr ? *SelectedUserFriendlyNamePtr : nullptr;
+		const TSharedPtr<FString> SelectedUserName = SelectedUserNamePtr ? *SelectedUserNamePtr : nullptr;
 
 		return SAssignNew(ComboBox, SComboBox<TSharedPtr<FString>>)
 			.ContentPadding(FMargin(2.0f, 2.0f, 2.0f, 1.0f))
-			.OptionsSource(&UserFriendlyNameList)
-			.InitiallySelectedItem(SelectedUserFriendlyName)
+			.OptionsSource(&UserNameList)
+			.InitiallySelectedItem(SelectedUserName)
 			.OnGenerateWidget(this, &SDMXPixelMappingComponentPin::GenerateComboBoxEntry)
 			.OnSelectionChanged(this, &SDMXPixelMappingComponentPin::ComboBoxSelectionChanged)
 			.IsEnabled(this, &SGraphPin::IsEditingEnabled)
 			.Visibility(this, &SGraphPin::GetDefaultValueVisibility)
 			[
 				SNew(STextBlock)
-				.Text(this, &SDMXPixelMappingComponentPin::GetSelectedUserFriendlyName)
+				.Text(this, &SDMXPixelMappingComponentPin::GetSelectedUserName)
 			];
 	}
 
@@ -90,31 +90,31 @@ protected:
 	 * @param NameItem The newly selected item in the combo box
 	 * @param SelectInfo Provides context on how the selection changed
 	 */
-	void ComboBoxSelectionChanged(TSharedPtr<FString> UserFriendlyNameItem, ESelectInfo::Type SelectInfo)
+	void ComboBoxSelectionChanged(TSharedPtr<FString> UserNameItem, ESelectInfo::Type SelectInfo)
 	{
-		if (GraphPinObj && UserFriendlyNameItem.IsValid())
+		if (GraphPinObj && UserNameItem.IsValid())
 		{
-			FString UserFriendlyName = *UserFriendlyNameItem.Get();
-			FName ComponentName = NameTable.GetComponentName(UserFriendlyName);
+			FString UserName = *UserNameItem.Get();
+			FName ComponentName = NameTable.GetComponentName(UserName);
 
 			SetNameToPin(ComponentName);
 		}
 	}
 	
 private:	
-	TSharedRef<SWidget> GenerateComboBoxEntry(TSharedPtr<FString> UserFriendlyName)
+	TSharedRef<SWidget> GenerateComboBoxEntry(TSharedPtr<FString> UserName)
 	{
 		return
 			SNew(STextBlock)
-			.Text(FText::FromString(*UserFriendlyName));
+			.Text(FText::FromString(*UserName));
 	}
 
-	FText GetSelectedUserFriendlyName() const
+	FText GetSelectedUserName() const
 	{
 		if (GraphPinObj)
 		{
 			FName Selection = FName(*GraphPinObj->GetDefaultAsString());
-			return FText::FromString(NameTable.GetUserFriendlyName(Selection));
+			return FText::FromString(NameTable.GetUserName(Selection));
 		}
 
 		return FText::GetEmpty();
@@ -143,15 +143,15 @@ private:
 	/** Reference to Combo Box object */
 	TSharedPtr<SComboBox<TSharedPtr<FString>>> ComboBox;
 
-	TArray<TSharedPtr<FString>> UserFriendlyNameList;
+	TArray<TSharedPtr<FString>> UserNameList;
 
-	/** Maps component FNames to the corresponding UserFriendly Name displayed in pixel mapping designer */
-	struct FComponentUserFriendlyNameTable
+	/** Maps component FNames to the corresponding User Name displayed in pixel mapping designer */
+	struct FComponentUserNameTable
 	{
-		FComponentUserFriendlyNameTable()
+		FComponentUserNameTable()
 		{}
 
-		FComponentUserFriendlyNameTable(UDMXPixelMapping* PixelMapping)
+		FComponentUserNameTable(UDMXPixelMapping* PixelMapping)
 		{
 			if (PixelMapping)
 			{
@@ -161,9 +161,9 @@ private:
 				{
 					if (Component)
 					{
-						// GetUserFriendlyName will display what is shown everywhere else in the UI
-						TTuple<FName, FString> Tuple = TTuple<FName, FString>(Component->GetFName(), Component->GetUserFriendlyName());
-						ComponentUserFriendlyNameArr.Add(Tuple);
+						// GetUserName will display what is shown everywhere else in the UI
+						TTuple<FName, FString> Tuple = TTuple<FName, FString>(Component->GetFName(), Component->GetUserName());
+						ComponentUserNameArr.Add(Tuple);
 					}
 				}
 			}
@@ -171,42 +171,42 @@ private:
 
 		bool Contains(const FName& ComponentName) const
 		{
-			int32 Index = ComponentUserFriendlyNameArr.IndexOfByPredicate([&ComponentName](const TTuple<FName, FString>& ComponentUserFriendlyKvp) {
-				return ComponentUserFriendlyKvp.Key == ComponentName;
+			int32 Index = ComponentUserNameArr.IndexOfByPredicate([&ComponentName](const TTuple<FName, FString>& ComponentUserKvp) {
+				return ComponentUserKvp.Key == ComponentName;
 				});
 			return Index != INDEX_NONE;
 		}
 
-		void GetUserFriendlyNames(TArray<TSharedPtr<FString>>& OutUserFriendlyNames)
+		void GetUserNames(TArray<TSharedPtr<FString>>& OutUserNames)
 		{
-			OutUserFriendlyNames.Reset();
-			for (const TTuple<FName, FString>& ComponentUserFriendlyKvp : ComponentUserFriendlyNameArr)
+			OutUserNames.Reset();
+			for (const TTuple<FName, FString>& ComponentUserKvp : ComponentUserNameArr)
 			{
-				OutUserFriendlyNames.Add(MakeShared<FString>(ComponentUserFriendlyKvp.Value));
+				OutUserNames.Add(MakeShared<FString>(ComponentUserKvp.Value));
 			}
 		}
 
-		FString GetUserFriendlyName(const FName& ComponentName) const
+		FString GetUserName(const FName& ComponentName) const
 		{
-			int32 Index = ComponentUserFriendlyNameArr.IndexOfByPredicate([&ComponentName](const TTuple<FName, FString>& ComponentUserFriendlyKvp) {
-				return ComponentUserFriendlyKvp.Key == ComponentName;
+			int32 Index = ComponentUserNameArr.IndexOfByPredicate([&ComponentName](const TTuple<FName, FString>& ComponentUserKvp) {
+				return ComponentUserKvp.Key == ComponentName;
 				});
 
 			if (Index != INDEX_NONE)
 			{
-				return ComponentUserFriendlyNameArr[Index].Value;
+				return ComponentUserNameArr[Index].Value;
 			}
 
 			return FString(TEXT("None"));
 		}
 
-		FName GetComponentName(const FString& UserFriendlyName) const
+		FName GetComponentName(const FString& UserName) const
 		{
-			for (const TTuple<FName, FString>& ComponentUserFriendlyKvp : ComponentUserFriendlyNameArr)
+			for (const TTuple<FName, FString>& ComponentUserKvp : ComponentUserNameArr)
 			{
-				if (ComponentUserFriendlyKvp.Value == UserFriendlyName)
+				if (ComponentUserKvp.Value == UserName)
 				{
-					return ComponentUserFriendlyKvp.Key;
+					return ComponentUserKvp.Key;
 				}
 			}
 
@@ -215,19 +215,19 @@ private:
 
 		int32 Num() const
 		{
-			return ComponentUserFriendlyNameArr.Num();
+			return ComponentUserNameArr.Num();
 		}
 
 		TTuple<FName, FString> First() const
 		{
-			check(ComponentUserFriendlyNameArr.Num() > 0);
-			return ComponentUserFriendlyNameArr[0];
+			check(ComponentUserNameArr.Num() > 0);
+			return ComponentUserNameArr[0];
 		}
 
 	private:
-		TArray<TTuple<FName, FString>> ComponentUserFriendlyNameArr;
+		TArray<TTuple<FName, FString>> ComponentUserNameArr;
 	};
 
 	/** List of available component names */
-	FComponentUserFriendlyNameTable NameTable;
+	FComponentUserNameTable NameTable;
 };
diff --git a/Engine/Plugins/VirtualProduction/DMX/DMXPixelMapping/Source/DMXPixelMappingEditor/Private/Customizations/DMXPixelMappingDetailCustomization_FixtureGroup.cpp b/Engine/Plugins/VirtualProduction/DMX/DMXPixelMapping/Source/DMXPixelMappingEditor/Private/Customizations/DMXPixelMappingDetailCustomization_FixtureGroup.cpp
--- a/Engine/Plugins/VirtualProduction/DMX/DMXPixelMapping/Source/DMXPixelMappingEditor/Private/Customizations/DMXPixelMappingDetailCustomization_FixtureGroup.cpp
+++ b/Engine/Plugins/VirtualProduction/DMX/DMXPixelMapping/Source/DMXPixelMappingEditor/Private/Customizations/DMXPixelMappingDetailCustomization_FixtureGroup.cpp
@@ -3,7 +3,6 @@
 #include "Customizations/DMXPixelMappingDetailCustomization_FixtureGroup.h"
 
 #include "Algo/Find.h"
-
 #include "Components/DMXPixelMappingFixtureGroupComponent.h"
 #include "DetailLayoutBuilder.h"
 #include "Editor.h"
@@ -19,6 +18,7 @@ void FDMXPixelMappingDetailCustomization_FixtureGroup::CustomizeDetails(IDetailL
 	PropertyUtilities = InDetailLayout.GetPropertyUtilities();
 
 	// Hide the Layout Script property (shown in its own panel, see SDMXPixelMappingLayoutView)
+	InDetailLayout.HideProperty(GET_MEMBER_NAME_CHECKED(UDMXPixelMappingFixtureGroupComponent, DMXLibrary));
 	InDetailLayout.HideProperty(GET_MEMBER_NAME_CHECKED(UDMXPixelMappingFixtureGroupComponent, LayoutScript));
 
 	// Handle size changes
diff --git a/Engine/Plugins/VirtualProduction/DMX/DMXPixelMapping/Source/DMXPixelMappingEditor/Private/Toolkits/DMXPixelMappingToolkit.cpp b/Engine/Plugins/VirtualProduction/DMX/DMXPixelMapping/Source/DMXPixelMappingEditor/Private/Toolkits/DMXPixelMappingToolkit.cpp
--- a/Engine/Plugins/VirtualProduction/DMX/DMXPixelMapping/Source/DMXPixelMappingEditor/Private/Toolkits/DMXPixelMappingToolkit.cpp
+++ b/Engine/Plugins/VirtualProduction/DMX/DMXPixelMapping/Source/DMXPixelMappingEditor/Private/Toolkits/DMXPixelMappingToolkit.cpp
@@ -263,10 +263,19 @@ bool FDMXPixelMappingToolkit::IsComponentSelected(UDMXPixelMappingBaseComponent*
 
 void FDMXPixelMappingToolkit::AddRenderer()
 {
-	if (DMXPixelMapping)
+	UDMXPixelMappingRootComponent* RootComponent = DMXPixelMapping ? DMXPixelMapping->GetRootComponent() : nullptr;
+	if (RootComponent)
 	{
-		UDMXPixelMappingRendererComponent* RendererComponent = FDMXPixelMappingEditorUtils::AddRenderer(DMXPixelMapping);
-		SetActiveRenderComponent(RendererComponent);
+		const FScopedTransaction AddMappingTransaction(LOCTEXT("AddMappingTransaction", "Add Mapping to Pixel Mapping"));
+
+		RootComponent->PreEditChange(UDMXPixelMappingBaseComponent::StaticClass()->FindPropertyByName(GET_MEMBER_NAME_CHECKED(UDMXPixelMappingBaseComponent, Children)));
+		UDMXPixelMappingRendererComponent* NewRendererComponent = FDMXPixelMappingEditorUtils::AddRenderer(DMXPixelMapping);
+		RootComponent->PostEditChange();
+
+		SetActiveRenderComponent(NewRendererComponent);
+
+		const FDMXPixelMappingComponentReference ComponentReference(StaticCastSharedRef<FDMXPixelMappingToolkit>(AsShared()), NewRendererComponent);
+		SelectComponents({ ComponentReference } );
 	}
 }
 
@@ -481,22 +490,20 @@ void FDMXPixelMappingToolkit::SizeSelectedComponentToTexture(bool bTransacted)
 	Component->PostEditChange();
 }
 
-void FDMXPixelMappingToolkit::OnComponentAdded(UDMXPixelMapping* PixelMapping, UDMXPixelMappingBaseComponent* Component)
+void FDMXPixelMappingToolkit::OnComponentAddedOrRemoved(UDMXPixelMapping* PixelMapping, UDMXPixelMappingBaseComponent* Component)
 {
-	if (!bAddingComponents)
+	if (!bAddingComponents && !bRemovingComponents)
 	{
 		UpdateBlueprintNodes();
 	}
 }
 
-void FDMXPixelMappingToolkit::OnComponentRemoved(UDMXPixelMapping* PixelMapping, UDMXPixelMappingBaseComponent* Component)
+void FDMXPixelMappingToolkit::OnComponentRenamed(UDMXPixelMappingBaseComponent* Component)
 {
-	if (!bRemovingComponents)
-	{
-		UpdateBlueprintNodes();
-	}
+	UpdateBlueprintNodes();
 }
 
+
 void FDMXPixelMappingToolkit::InitializeInternal(const EToolkitMode::Type Mode, const TSharedPtr<class IToolkitHost>& InitToolkitHost, const FGuid& MessageLogGuid)
 {
 	if (!DMXPixelMapping)
@@ -508,8 +515,9 @@ void FDMXPixelMappingToolkit::InitializeInternal(const EToolkitMode::Type Mode,
 	DMXPixelMapping->CreateOrLoadObjects();
 
 	// Bind to component changes
-	UDMXPixelMappingBaseComponent::GetOnComponentAdded().AddSP(this, &FDMXPixelMappingToolkit::OnComponentAdded);
-	UDMXPixelMappingBaseComponent::GetOnComponentRemoved().AddSP(this, &FDMXPixelMappingToolkit::OnComponentRemoved);
+	UDMXPixelMappingBaseComponent::GetOnComponentAdded().AddSP(this, &FDMXPixelMappingToolkit::OnComponentAddedOrRemoved);
+	UDMXPixelMappingBaseComponent::GetOnComponentRemoved().AddSP(this, &FDMXPixelMappingToolkit::OnComponentAddedOrRemoved);
+	UDMXPixelMappingBaseComponent::GetOnComponentRenamed().AddSP(this, &FDMXPixelMappingToolkit::OnComponentRenamed);
 
 	// Create commands
 	DesignerCommandList = MakeShareable(new FUICommandList);
@@ -593,7 +601,7 @@ void FDMXPixelMappingToolkit::InitializeInternal(const EToolkitMode::Type Mode,
 	ExtendToolbar();
 	RegenerateMenusAndToolbars();
 	
-	// Select the first renderer component
+	// Make an initial selection
 	if (UDMXPixelMappingRootComponent* RootComponent = DMXPixelMapping->GetRootComponent())
 	{
 		UDMXPixelMappingBaseComponent* const* FirstRendererComponetPtr = Algo::FindByPredicate(DMXPixelMapping->GetRootComponent()->GetChildren(), [](UDMXPixelMappingBaseComponent* Component)
@@ -602,8 +610,18 @@ void FDMXPixelMappingToolkit::InitializeInternal(const EToolkitMode::Type Mode,
 			});
 		if (FirstRendererComponetPtr)
 		{
-			const FDMXPixelMappingComponentReference ComponentReference(StaticCastSharedRef<FDMXPixelMappingToolkit>(AsShared()), *FirstRendererComponetPtr);
-			SelectComponents(TSet<FDMXPixelMappingComponentReference>({ ComponentReference }));
+			UDMXPixelMappingBaseComponent* const* FirstFixtureGroupComponetPtr = Algo::FindByPredicate((*FirstRendererComponetPtr)->GetChildren(), [](UDMXPixelMappingBaseComponent* Component)
+				{
+					return Component && Component->GetClass() == UDMXPixelMappingFixtureGroupComponent::StaticClass();
+				});
+
+			if (UDMXPixelMappingBaseComponent* const* ComponentToSelectPtr = FirstFixtureGroupComponetPtr ? 
+				FirstFixtureGroupComponetPtr : 
+				FirstRendererComponetPtr)
+			{
+				const FDMXPixelMappingComponentReference ComponentReference(StaticCastSharedRef<FDMXPixelMappingToolkit>(AsShared()), *ComponentToSelectPtr);
+				SelectComponents(TSet<FDMXPixelMappingComponentReference>({ ComponentReference }));
+			}
 		}
 	}
 }
diff --git a/Engine/Plugins/VirtualProduction/DMX/DMXPixelMapping/Source/DMXPixelMappingEditor/Private/Toolkits/DMXPixelMappingToolkit.h b/Engine/Plugins/VirtualProduction/DMX/DMXPixelMapping/Source/DMXPixelMappingEditor/Private/Toolkits/DMXPixelMappingToolkit.h
--- a/Engine/Plugins/VirtualProduction/DMX/DMXPixelMapping/Source/DMXPixelMappingEditor/Private/Toolkits/DMXPixelMappingToolkit.h
+++ b/Engine/Plugins/VirtualProduction/DMX/DMXPixelMapping/Source/DMXPixelMappingEditor/Private/Toolkits/DMXPixelMappingToolkit.h
@@ -118,12 +118,7 @@ public:
 	/** Adds a renderer to the pixel mapping. Sets it as active renderer if a new renderer was created */
 	void AddRenderer();
 
-	/**
-	 * Rename Pixel Maping component.
-	 *
-	 * @param PreviousObjectName	The previous object name, used to find the object
-	 * @param DesiredObjectName		The desired object name
-	 */
+	UE_DEPRECATED(5.3, "Removed without replacement. Use UObject::Rename instead")
 	void RenameComponent(const FName& PreviousObjectName, const FString& DesiredObjectName) const;
 
 	/** 
@@ -142,11 +137,11 @@ public:
 	void SizeSelectedComponentToTexture(bool bTransacted);
 
 private:
-	/** Called when a Component was added to the pixel mapping */
-	void OnComponentAdded(UDMXPixelMapping* PixelMapping, UDMXPixelMappingBaseComponent* Component);
+	/** Called when a component was added to the pixel mapping */
+	void OnComponentAddedOrRemoved(UDMXPixelMapping* PixelMapping, UDMXPixelMappingBaseComponent* Component);
 
-	/** Called when a Component was removed from the pixel mapping */
-	void OnComponentRemoved(UDMXPixelMapping* PixelMapping, UDMXPixelMappingBaseComponent* Component);
+	/** Called when a component object was renamed */
+	void OnComponentRenamed(UDMXPixelMappingBaseComponent* Component);
 
 	void PlayDMX();
 
diff --git a/Engine/Plugins/VirtualProduction/DMX/DMXPixelMapping/Source/DMXPixelMappingEditor/Private/ViewModels/DMXPixelMappingDMXLibraryViewModel.cpp b/Engine/Plugins/VirtualProduction/DMX/DMXPixelMapping/Source/DMXPixelMappingEditor/Private/ViewModels/DMXPixelMappingDMXLibraryViewModel.cpp
--- a/Engine/Plugins/VirtualProduction/DMX/DMXPixelMapping/Source/DMXPixelMappingEditor/Private/ViewModels/DMXPixelMappingDMXLibraryViewModel.cpp
+++ b/Engine/Plugins/VirtualProduction/DMX/DMXPixelMapping/Source/DMXPixelMappingEditor/Private/ViewModels/DMXPixelMappingDMXLibraryViewModel.cpp
@@ -2,6 +2,7 @@
 
 #include "DMXPixelMappingDMXLibraryViewModel.h"
 
+#include "Algo/Transform.h"
 #include "Components/DMXPixelMappingBaseComponent.h"
 #include "Components/DMXPixelMappingFixtureGroupComponent.h"
 #include "Components/DMXPixelMappingFixtureGroupItemComponent.h"
@@ -78,14 +79,15 @@ void UDMXPixelMappingDMXLibraryViewModel::CreateAndSetNewFixtureGroup(TWeakPtr<F
 
 void UDMXPixelMappingDMXLibraryViewModel::UpdateFixtureGroupFromSelection(TWeakPtr<FDMXPixelMappingToolkit> InWeakToolkit)
 {
-	WeakToolkit = InWeakToolkit;
-	if (!WeakToolkit.IsValid())
+	if (!InWeakToolkit.IsValid())
 	{
 		return;
 	}
+	WeakToolkit = InWeakToolkit;
 	const TSharedRef<FDMXPixelMappingToolkit> Toolkit = WeakToolkit.Pin().ToSharedRef();
-	const TSet<FDMXPixelMappingComponentReference>& SelectedComponents = Toolkit->GetSelectedComponents();
 
+	// Gather fixture groups present in the current selection
+	const TSet<FDMXPixelMappingComponentReference>& SelectedComponents = Toolkit->GetSelectedComponents();
 	TArray<UDMXPixelMappingFixtureGroupComponent*> FixtureGroupComponents;
 	for (const FDMXPixelMappingComponentReference& ComponentReference : SelectedComponents)
 	{
@@ -98,7 +100,7 @@ void UDMXPixelMappingDMXLibraryViewModel::UpdateFixtureGroupFromSelection(TWeakP
 		}
 	}
 
-	// Only allow a single selection 
+	// Update members, only allow a single selection. 
 	if (FixtureGroupComponents.IsEmpty() || FixtureGroupComponents.Num() > 1)
 	{
 		DMXLibrary = nullptr;
@@ -115,7 +117,8 @@ void UDMXPixelMappingDMXLibraryViewModel::UpdateFixtureGroupFromSelection(TWeakP
 
 void UDMXPixelMappingDMXLibraryViewModel::AddFixturePatchesEnsured(const TArray<TSharedPtr<FDMXEntityFixturePatchRef>>& FixturePatches)
 {
-	if (FixturePatches.IsEmpty() || !WeakToolkit.IsValid())
+	const TSharedPtr<FDMXPixelMappingToolkit> Toolkit = WeakToolkit.Pin();
+	if (FixturePatches.IsEmpty() || !Toolkit.IsValid())
 	{
 		return;
 	}
@@ -170,7 +173,7 @@ void UDMXPixelMappingDMXLibraryViewModel::AddFixturePatchesEnsured(const TArray<
 		}
 	}
 
-	const TArray<UDMXPixelMappingBaseComponent*> NewComponents = WeakToolkit.Pin()->CreateComponentsFromTemplates(RootComponent, FixtureGroupComponent, Templates);
+	const TArray<UDMXPixelMappingBaseComponent*> NewComponents = Toolkit->CreateComponentsFromTemplates(RootComponent, FixtureGroupComponent, Templates);
 
 	// Layout
 	const TArray<UDMXEntityFixturePatch*> FixturePatchesInLibrary = CommonDMXLibrary->GetEntitiesTypeCast<UDMXEntityFixturePatch>();
@@ -183,6 +186,15 @@ void UDMXPixelMappingDMXLibraryViewModel::AddFixturePatchesEnsured(const TArray<
 	{
 		LayoutAfterLastPatch(NewComponents);
 	}
+
+	// Select new components
+	TSet<FDMXPixelMappingComponentReference> ComponentReferencesToSelect;
+	Algo::Transform(NewComponents, ComponentReferencesToSelect,
+		[Toolkit](UDMXPixelMappingBaseComponent* Component)
+		{
+			return FDMXPixelMappingComponentReference(Toolkit, Component);
+		});
+	Toolkit->SelectComponents(ComponentReferencesToSelect);
 }
 
 void UDMXPixelMappingDMXLibraryViewModel::SaveFixturePatchListDescriptor(const FDMXReadOnlyFixturePatchListDescriptor& NewDescriptor)
diff --git a/Engine/Plugins/VirtualProduction/DMX/DMXPixelMapping/Source/DMXPixelMappingEditor/Private/ViewModels/DMXPixelMappingHierarchyItem.cpp b/Engine/Plugins/VirtualProduction/DMX/DMXPixelMapping/Source/DMXPixelMappingEditor/Private/ViewModels/DMXPixelMappingHierarchyItem.cpp
--- a/Engine/Plugins/VirtualProduction/DMX/DMXPixelMapping/Source/DMXPixelMappingEditor/Private/ViewModels/DMXPixelMappingHierarchyItem.cpp
+++ b/Engine/Plugins/VirtualProduction/DMX/DMXPixelMapping/Source/DMXPixelMappingEditor/Private/ViewModels/DMXPixelMappingHierarchyItem.cpp
@@ -51,7 +51,7 @@ FLinearColor FDMXPixelMappingHierarchyItem::GetEditorColor() const
 FText FDMXPixelMappingHierarchyItem::GetComponentNameText() const
 {
 	return WeakComponent.IsValid() ?
-		FText::FromString(WeakComponent->GetUserFriendlyName()) :
+		FText::FromString(WeakComponent->GetUserName()) :
 		FText::GetEmpty();
 }
 
diff --git a/Engine/Plugins/VirtualProduction/DMX/DMXPixelMapping/Source/DMXPixelMappingEditor/Private/ViewModels/DMXPixelMappingOutputComponentModel.cpp b/Engine/Plugins/VirtualProduction/DMX/DMXPixelMapping/Source/DMXPixelMappingEditor/Private/ViewModels/DMXPixelMappingOutputComponentModel.cpp
--- a/Engine/Plugins/VirtualProduction/DMX/DMXPixelMapping/Source/DMXPixelMappingEditor/Private/ViewModels/DMXPixelMappingOutputComponentModel.cpp
+++ b/Engine/Plugins/VirtualProduction/DMX/DMXPixelMapping/Source/DMXPixelMappingEditor/Private/ViewModels/DMXPixelMappingOutputComponentModel.cpp
@@ -65,7 +65,7 @@ FText FDMXPixelMappingOutputComponentModel::GetName() const
 {
 	if (UDMXPixelMappingOutputComponent* OutputComponent = Cast<UDMXPixelMappingOutputComponent>(WeakOutputComponent.Get()))
 	{
-		return FText::FromString(OutputComponent->GetUserFriendlyName());
+		return FText::FromString(OutputComponent->GetUserName());
 	}
 
 	return FText::GetEmpty();
diff --git a/Engine/Plugins/VirtualProduction/DMX/DMXPixelMapping/Source/DMXPixelMappingEditor/Private/Views/SDMXPixelMappingDMXLibraryView.cpp b/Engine/Plugins/VirtualProduction/DMX/DMXPixelMapping/Source/DMXPixelMappingEditor/Private/Views/SDMXPixelMappingDMXLibraryView.cpp
--- a/Engine/Plugins/VirtualProduction/DMX/DMXPixelMapping/Source/DMXPixelMappingEditor/Private/Views/SDMXPixelMappingDMXLibraryView.cpp
+++ b/Engine/Plugins/VirtualProduction/DMX/DMXPixelMapping/Source/DMXPixelMappingEditor/Private/Views/SDMXPixelMappingDMXLibraryView.cpp
@@ -56,7 +56,7 @@ void SDMXPixelMappingDMXLibraryView::Construct(const FArguments& InArgs, const T
 	ModelDetailsView = PropertyEditorModule.CreateDetailView(DetailsViewArgs);
 	ModelDetailsView->SetObject(ViewModel);
 
-	const TAttribute<bool> AddDMXLibraryButtonEnabledAttribute = TAttribute<bool>::CreateLambda(
+	const TAttribute<bool> AddFixtureGroupEnabledAttribute = TAttribute<bool>::CreateLambda(
 		[this]()
 		{
 			return bRenderComponentContainedInSelection;
@@ -81,12 +81,12 @@ void SDMXPixelMappingDMXLibraryView::Construct(const FArguments& InArgs, const T
 				.ToolTipText_Lambda([this]()
 					{
 						return bRenderComponentContainedInSelection ?
-							LOCTEXT("AddFixtureGroupTooltip", "Adds a DMX Library to the Pixel Mapping") :
+							LOCTEXT("AddFixtureGroupTooltip", "Adds a Fixture Group to the Pixel Mapping. In the fixture group, a DMX Library can be selected.") :
 							LOCTEXT("CannotAddFixtureGroupTooltip", "Please select a source texture");
 					})
 				.ContentPadding(FMargin(5.0f, 1.0f))
-				.OnClicked(this, &SDMXPixelMappingDMXLibraryView::OnAddDMXLibraryButtonClicked)
-				.IsEnabled(AddDMXLibraryButtonEnabledAttribute)
+				.OnClicked(this, &SDMXPixelMappingDMXLibraryView::OnAddFixtureGroupButtonClicked)
+				.IsEnabled(AddFixtureGroupEnabledAttribute)
 				.Content()
 				[
 					SNew(SHorizontalBox)
@@ -105,7 +105,7 @@ void SDMXPixelMappingDMXLibraryView::Construct(const FArguments& InArgs, const T
 					.Padding(FMargin(2.f, 0.f, 2.f, 0.f))
 					[
 						SNew(STextBlock)
-						.Text(LOCTEXT("AddFixtureGroupLabel", "Add DMX Library"))
+						.Text(LOCTEXT("AddFixtureGroupLabel", "Add Fixture Group"))
 					]
 				]
 			]
@@ -254,7 +254,7 @@ void SDMXPixelMappingDMXLibraryView::ForceRefresh()
 		{
 			if (EditedFixtureGroup)
 			{
-				return FText::FromString(EditedFixtureGroup->GetUserFriendlyName());
+				return FText::FromString(EditedFixtureGroup->GetUserName());
 			}
 			else if (ViewModel->IsMoreThanOneFixtureGroupSelected())
 			{
@@ -335,7 +335,7 @@ void SDMXPixelMappingDMXLibraryView::OnComponentAddedOrRemoved(UDMXPixelMapping*
 	RequestRefresh();
 }
 
-FReply SDMXPixelMappingDMXLibraryView::OnAddDMXLibraryButtonClicked()
+FReply SDMXPixelMappingDMXLibraryView::OnAddFixtureGroupButtonClicked()
 {
 	if (ensureMsgf(ViewModel, TEXT("Invalid view model for PixelMapping DMX Library View, cannot display view.")))
 	{
@@ -349,7 +349,8 @@ FReply SDMXPixelMappingDMXLibraryView::OnAddDMXLibraryButtonClicked()
 
 FReply SDMXPixelMappingDMXLibraryView::OnAddSelectedPatchesClicked()
 {
-	if (ensureMsgf(ViewModel, TEXT("Invalid view model for PixelMapping DMX Library View, cannot display view.")))
+	if (WeakToolkit.IsValid() &&
+		ensureMsgf(ViewModel, TEXT("Invalid view model for PixelMapping DMX Library View, cannot display view.")))
 	{
 		const FScopedTransaction AddSelectedFixturePatchesTransaction(LOCTEXT("AddSelectedFixturePatchesTransaction", "Add Fixture Patches to Pixel Mapping"));
 
diff --git a/Engine/Plugins/VirtualProduction/DMX/DMXPixelMapping/Source/DMXPixelMappingEditor/Private/Views/SDMXPixelMappingDMXLibraryView.h b/Engine/Plugins/VirtualProduction/DMX/DMXPixelMapping/Source/DMXPixelMappingEditor/Private/Views/SDMXPixelMappingDMXLibraryView.h
--- a/Engine/Plugins/VirtualProduction/DMX/DMXPixelMapping/Source/DMXPixelMappingEditor/Private/Views/SDMXPixelMappingDMXLibraryView.h
+++ b/Engine/Plugins/VirtualProduction/DMX/DMXPixelMapping/Source/DMXPixelMappingEditor/Private/Views/SDMXPixelMappingDMXLibraryView.h
@@ -69,8 +69,8 @@ private:
 	/** Called when a component was added or removed from pixel mapping */
 	void OnComponentAddedOrRemoved(UDMXPixelMapping* PixelMapping, UDMXPixelMappingBaseComponent* Component);
 
-	/** Called when the 'Add DMX Library' button was clicked */
-	FReply OnAddDMXLibraryButtonClicked();
+	/** Called when the 'Add Fixture Group' button was clicked */
+	FReply OnAddFixtureGroupButtonClicked();
 
 	/** Called when the 'Add Selected Patches' button was clicked */
 	FReply OnAddSelectedPatchesClicked();
diff --git a/Engine/Plugins/VirtualProduction/DMX/DMXPixelMapping/Source/DMXPixelMappingEditor/Private/Views/SDMXPixelMappingDesignerView.cpp b/Engine/Plugins/VirtualProduction/DMX/DMXPixelMapping/Source/DMXPixelMappingEditor/Private/Views/SDMXPixelMappingDesignerView.cpp
--- a/Engine/Plugins/VirtualProduction/DMX/DMXPixelMapping/Source/DMXPixelMappingEditor/Private/Views/SDMXPixelMappingDesignerView.cpp
+++ b/Engine/Plugins/VirtualProduction/DMX/DMXPixelMapping/Source/DMXPixelMappingEditor/Private/Views/SDMXPixelMappingDesignerView.cpp
@@ -799,7 +799,7 @@ FText SDMXPixelMappingDesignerView::GetHoveredComponentNameText() const
 {
 	if (UDMXPixelMappingOutputComponent* ComponentUnderCursor = GetComponentUnderCursor())
 	{
-		return FText::FromString(ComponentUnderCursor->GetUserFriendlyName());
+		return FText::FromString(ComponentUnderCursor->GetUserName());
 	}
 
 	return FText();
@@ -811,7 +811,7 @@ FText SDMXPixelMappingDesignerView::GetHoveredComponentParentNameText() const
 	{
 		if (UDMXPixelMappingBaseComponent* ParentOfComponentUnderCursor = ComponentUnderCursor->GetParent())
 		{
-			return FText::FromString(ParentOfComponentUnderCursor->GetUserFriendlyName());
+			return FText::FromString(ParentOfComponentUnderCursor->GetUserName());
 		}
 	}
 
diff --git a/Engine/Plugins/VirtualProduction/DMX/DMXPixelMapping/Source/DMXPixelMappingEditor/Private/Views/SDMXPixelMappingHierarchyView.cpp b/Engine/Plugins/VirtualProduction/DMX/DMXPixelMapping/Source/DMXPixelMappingEditor/Private/Views/SDMXPixelMappingHierarchyView.cpp
--- a/Engine/Plugins/VirtualProduction/DMX/DMXPixelMapping/Source/DMXPixelMappingEditor/Private/Views/SDMXPixelMappingHierarchyView.cpp
+++ b/Engine/Plugins/VirtualProduction/DMX/DMXPixelMapping/Source/DMXPixelMappingEditor/Private/Views/SDMXPixelMappingHierarchyView.cpp
@@ -508,15 +508,27 @@ void SDMXPixelMappingHierarchyView::OnEditorSelectionChanged()
 
 void SDMXPixelMappingHierarchyView::AdoptSelectionFromToolkit()
 {
-	if (bIsUpdatingSelection)
+	if (bIsUpdatingSelection || !HierarchyTreeView.IsValid())
 	{
 		return;
 	}
 
+	const TArray<TSharedPtr<FDMXPixelMappingHierarchyItem>> OldSelection = HierarchyTreeView->GetSelectedItems();
 	for (const TSharedPtr<FDMXPixelMappingHierarchyItem>& RootItem : AllRootItems)
 	{
 		RecursiveAdoptSelectionFromToolkit(RootItem.ToSharedRef());
 	}
+
+	const TArray<TSharedPtr<FDMXPixelMappingHierarchyItem>> NewSelection = HierarchyTreeView->GetSelectedItems();
+	const bool bOldSelectionContainsNewSelection = Algo::FindByPredicate(NewSelection, [&OldSelection](const TSharedPtr<FDMXPixelMappingHierarchyItem>& SelectedItem)
+		{
+			return OldSelection.Contains(SelectedItem);
+		}) != nullptr;
+
+	if (!bOldSelectionContainsNewSelection && !NewSelection.IsEmpty())
+	{
+		HierarchyTreeView->RequestScrollIntoView(NewSelection[0]);
+	}
 }
 
 void SDMXPixelMappingHierarchyView::RecursiveAdoptSelectionFromToolkit(const TSharedRef<FDMXPixelMappingHierarchyItem>& Item)
diff --git a/Engine/Plugins/VirtualProduction/DMX/DMXPixelMapping/Source/DMXPixelMappingEditor/Private/Views/SDMXPixelMappingPreviewView.cpp b/Engine/Plugins/VirtualProduction/DMX/DMXPixelMapping/Source/DMXPixelMappingEditor/Private/Views/SDMXPixelMappingPreviewView.cpp
--- a/Engine/Plugins/VirtualProduction/DMX/DMXPixelMapping/Source/DMXPixelMappingEditor/Private/Views/SDMXPixelMappingPreviewView.cpp
+++ b/Engine/Plugins/VirtualProduction/DMX/DMXPixelMapping/Source/DMXPixelMappingEditor/Private/Views/SDMXPixelMappingPreviewView.cpp
@@ -391,7 +391,7 @@ FText SDMXPixelMappingPreviewView::GetSelectedComponentNameText() const
 	UDMXPixelMappingBaseComponent* BaseComponent = SelectedComponentRef.GetComponent();
 	if (BaseComponent)
 	{
-		return FText::FromString(BaseComponent->GetUserFriendlyName());
+		return FText::FromString(BaseComponent->GetUserName());
 	}
 
 	return FText();
@@ -403,7 +403,7 @@ FText SDMXPixelMappingPreviewView::GetSelectedComponentParentNameText() const
 	UDMXPixelMappingBaseComponent* BaseComponent = SelectedComponentRef.GetComponent();
 	if (BaseComponent && BaseComponent->GetParent())
 	{
-		return FText::FromString(BaseComponent->GetParent()->GetUserFriendlyName());
+		return FText::FromString(BaseComponent->GetParent()->GetUserName());
 	}
 
 	return FText();
diff --git a/Engine/Plugins/VirtualProduction/DMX/DMXPixelMapping/Source/DMXPixelMappingEditor/Private/Widgets/SDMXPixelMappingHierarchyRow.cpp b/Engine/Plugins/VirtualProduction/DMX/DMXPixelMapping/Source/DMXPixelMappingEditor/Private/Widgets/SDMXPixelMappingHierarchyRow.cpp
--- a/Engine/Plugins/VirtualProduction/DMX/DMXPixelMapping/Source/DMXPixelMappingEditor/Private/Widgets/SDMXPixelMappingHierarchyRow.cpp
+++ b/Engine/Plugins/VirtualProduction/DMX/DMXPixelMapping/Source/DMXPixelMappingEditor/Private/Widgets/SDMXPixelMappingHierarchyRow.cpp
@@ -4,6 +4,7 @@
 
 #include "Components/DMXPixelMappingOutputComponent.h"
 #include "DMXEditorStyle.h"
+#include "DMXPixelMapping.h"
 #include "DMXPixelMappingEditorUtils.h"
 #include "DragDrop/DMXPixelMappingDragDropOp.h"
 #include "ScopedTransaction.h"
@@ -106,6 +107,8 @@ TSharedRef<SWidget> SDMXPixelMappingHierarchyRow::GenerateComponentNameWidget()
 			SAssignNew(EditableNameTextBox, SInlineEditableTextBlock)
 			.Font(FAppStyle::GetFontStyle(TEXT("PropertyWindow.NormalFont")))
 			.Text(Item.Get(), &FDMXPixelMappingHierarchyItem::GetComponentNameText)
+			.OnVerifyTextChanged(this, &SDMXPixelMappingHierarchyRow::OnVerifyNameTextChanged)
+			.OnTextCommitted(this, &SDMXPixelMappingHierarchyRow::OnNameTextCommited)
 			.ColorAndOpacity(FLinearColor::White)
 		];
 }
@@ -142,16 +145,57 @@ TSharedRef<SWidget> SDMXPixelMappingHierarchyRow::GeneratePatchWidget()
 		];
 }
 
+bool SDMXPixelMappingHierarchyRow::OnVerifyNameTextChanged(const FText& InText, FText& OutErrorMessage)
+{
+	UDMXPixelMapping* PixelMapping = WeakToolkit.IsValid() ? WeakToolkit.Pin()->GetDMXPixelMapping() : nullptr;
+	const UDMXPixelMappingBaseComponent* MyComponent = Item->GetComponent();
+	if (PixelMapping && MyComponent)
+	{
+		TArray<const UDMXPixelMappingBaseComponent*> AllComponents;
+		PixelMapping->GetAllComponentsOfClass(AllComponents);
+
+		const UClass* MyComponentClass = MyComponent->GetClass();
+		const FString DesiredName = InText.ToString();
+		
+		// Allow to clear the name to revert to the generated name
+		if (DesiredName.IsEmpty())
+		{
+			return true;
+		}
+
+		const bool bUniqueName = Algo::FindByPredicate(AllComponents, [MyComponent, MyComponentClass, &DesiredName](const UDMXPixelMappingBaseComponent* OtherComponent)
+			{
+				return
+					OtherComponent &&
+					OtherComponent != MyComponent &&
+					OtherComponent->GetUserName() == DesiredName;
+			}) == nullptr;
+		if (!bUniqueName)
+		{
+			OutErrorMessage = LOCTEXT("RenameComponentTransaction", "A component with this name already exists");
+		}
+
+		return bUniqueName;
+	}
+
+	return true;
+}
+
 void SDMXPixelMappingHierarchyRow::OnNameTextCommited(const FText& InText, ETextCommit::Type CommitInfo)
 {
 	const TSharedPtr<FDMXPixelMappingToolkit> Toolkit = WeakToolkit.Pin();
-	const UDMXPixelMappingBaseComponent* Component = Item.IsValid() ? Item->GetComponent() : nullptr;
-	if (Toolkit.IsValid() && Component && !InText.IsEmpty())
+	UDMXPixelMappingBaseComponent* Component = Item.IsValid() ? Item->GetComponent() : nullptr;
+	if (Toolkit.IsValid() && Component)
 	{
 		const FScopedTransaction RenameComponentTransaction(LOCTEXT("RenameComponentTransaction", "Rename Pixel Mapping Component"));
 
-		const FName& ComponentName = Component->GetFName();
-		Toolkit->RenameComponent(ComponentName, InText.ToString());
+		Component->Modify();
+		Component->SetUserName(InText.ToString());
+
+		// To ease debugging also rename the UObject
+		const FName ObjectName = MakeObjectNameFromDisplayLabel(InText.ToString(), Component->GetFName());
+		const FName UniqueObjectName = MakeUniqueObjectName(Component->GetOuter(), Component->GetClass(), ObjectName);
+		Component->Rename(*UniqueObjectName.ToString());
 	}
 }
 
diff --git a/Engine/Plugins/VirtualProduction/DMX/DMXPixelMapping/Source/DMXPixelMappingEditor/Private/Widgets/SDMXPixelMappingHierarchyRow.h b/Engine/Plugins/VirtualProduction/DMX/DMXPixelMapping/Source/DMXPixelMappingEditor/Private/Widgets/SDMXPixelMappingHierarchyRow.h
--- a/Engine/Plugins/VirtualProduction/DMX/DMXPixelMapping/Source/DMXPixelMappingEditor/Private/Widgets/SDMXPixelMappingHierarchyRow.h
+++ b/Engine/Plugins/VirtualProduction/DMX/DMXPixelMapping/Source/DMXPixelMappingEditor/Private/Widgets/SDMXPixelMappingHierarchyRow.h
@@ -49,6 +49,9 @@ private:
 	/** Generates a widget to display the patch */
 	TSharedRef<SWidget> GeneratePatchWidget();
 
+	/** Called to verify the changed name in the related editable text block */
+	bool OnVerifyNameTextChanged(const FText& InText, FText& OutErrorMessage);
+
 	/** Called when text is committed on the node */
 	void OnNameTextCommited(const FText& InText, ETextCommit::Type CommitInfo);
 
diff --git a/Engine/Plugins/VirtualProduction/DMX/DMXPixelMapping/Source/DMXPixelMappingRuntime/Private/Components/DMXPixelMappingBaseComponent.cpp b/Engine/Plugins/VirtualProduction/DMX/DMXPixelMapping/Source/DMXPixelMappingRuntime/Private/Components/DMXPixelMappingBaseComponent.cpp
--- a/Engine/Plugins/VirtualProduction/DMX/DMXPixelMapping/Source/DMXPixelMappingRuntime/Private/Components/DMXPixelMappingBaseComponent.cpp
+++ b/Engine/Plugins/VirtualProduction/DMX/DMXPixelMapping/Source/DMXPixelMappingRuntime/Private/Components/DMXPixelMappingBaseComponent.cpp
@@ -50,7 +50,7 @@ void UDMXPixelMappingBaseComponent::PostRename(UObject* OldOuter, const FName Ol
 {
 	Super::PostRename(OldOuter, OldName);
 
-	OnComponentRenamed.Broadcast(GetPixelMapping(), this, OldOuter, OldName);
+	OnComponentRenamed.Broadcast(this);
 }
 
 #if WITH_EDITOR
@@ -185,7 +185,7 @@ UDMXPixelMappingBaseComponent* UDMXPixelMappingBaseComponent::GetChildAt(int32 I
 void UDMXPixelMappingBaseComponent::AddChild(UDMXPixelMappingBaseComponent* InComponent)
 {
 #if WITH_EDITOR
-	ensureMsgf(InComponent, TEXT("Trying to add nullptr to %s"), *GetUserFriendlyName());
+	ensureMsgf(InComponent, TEXT("Trying to add nullptr to %s"), *GetUserName());
 #endif 
 
 	if (InComponent)
@@ -228,7 +228,17 @@ void UDMXPixelMappingBaseComponent::ClearChildren()
 
 FString UDMXPixelMappingBaseComponent::GetUserFriendlyName() const
 {
-	return GetName();
+	return GetUserName();
+}
+
+FString UDMXPixelMappingBaseComponent::GetUserName() const
+{
+	return UserName.IsEmpty() ? GetName() : UserName;
+}
+
+void UDMXPixelMappingBaseComponent::SetUserName(const FString& NewName)
+{
+	UserName = NewName;
 }
 
 void UDMXPixelMappingBaseComponent::GetChildComponentsRecursively(TArray<UDMXPixelMappingBaseComponent*>& Components)
diff --git a/Engine/Plugins/VirtualProduction/DMX/DMXPixelMapping/Source/DMXPixelMappingRuntime/Private/Components/DMXPixelMappingFixtureGroupComponent.cpp b/Engine/Plugins/VirtualProduction/DMX/DMXPixelMapping/Source/DMXPixelMappingRuntime/Private/Components/DMXPixelMappingFixtureGroupComponent.cpp
--- a/Engine/Plugins/VirtualProduction/DMX/DMXPixelMapping/Source/DMXPixelMappingRuntime/Private/Components/DMXPixelMappingFixtureGroupComponent.cpp
+++ b/Engine/Plugins/VirtualProduction/DMX/DMXPixelMapping/Source/DMXPixelMappingRuntime/Private/Components/DMXPixelMappingFixtureGroupComponent.cpp
@@ -199,14 +199,17 @@ void UDMXPixelMappingFixtureGroupComponent::HandlePositionChanged()
 	LastPosition = GetPosition();
 }
 
-FString UDMXPixelMappingFixtureGroupComponent::GetUserFriendlyName() const
+FString UDMXPixelMappingFixtureGroupComponent::GetUserName() const
 {
-	if (DMXLibrary)
+	if (UserName.IsEmpty())
 	{
-		return FString::Printf(TEXT("Fixture Group: %s"), *DMXLibrary->GetName());
+		const FString DMXLibraryName = DMXLibrary ? DMXLibrary->GetName() : LOCTEXT("NoDMXLibrary", "No DMX Library").ToString();
+		return FString::Printf(TEXT("Fixture Group: %s"), *DMXLibraryName);
+	}
+	else
+	{
+		return UserName;
 	}
-
-	return FString("Fixture Group: No Library");
 }
 
 #if WITH_EDITOR
diff --git a/Engine/Plugins/VirtualProduction/DMX/DMXPixelMapping/Source/DMXPixelMappingRuntime/Private/Components/DMXPixelMappingFixtureGroupItemComponent.cpp b/Engine/Plugins/VirtualProduction/DMX/DMXPixelMapping/Source/DMXPixelMappingRuntime/Private/Components/DMXPixelMappingFixtureGroupItemComponent.cpp
--- a/Engine/Plugins/VirtualProduction/DMX/DMXPixelMapping/Source/DMXPixelMappingRuntime/Private/Components/DMXPixelMappingFixtureGroupItemComponent.cpp
+++ b/Engine/Plugins/VirtualProduction/DMX/DMXPixelMapping/Source/DMXPixelMappingRuntime/Private/Components/DMXPixelMappingFixtureGroupItemComponent.cpp
@@ -213,14 +213,17 @@ void UDMXPixelMappingFixtureGroupItemComponent::PostEditChangeChainProperty(FPro
 }
 #endif // WITH_EDITOR
 
-FString UDMXPixelMappingFixtureGroupItemComponent::GetUserFriendlyName() const
+FString UDMXPixelMappingFixtureGroupItemComponent::GetUserName() const
 {
-	if (UDMXEntityFixturePatch* Patch = FixturePatchRef.GetFixturePatch())
+	const UDMXEntityFixturePatch* Patch = FixturePatchRef.GetFixturePatch();
+	if (Patch && UserName.IsEmpty())
 	{
 		return Patch->GetDisplayName();
 	}
-
-	return FString(TEXT("Fixture Group Item: No Fixture Patch"));
+	else
+	{
+		return UserName;
+	}
 }
 
 const FName& UDMXPixelMappingFixtureGroupItemComponent::GetNamePrefix()
diff --git a/Engine/Plugins/VirtualProduction/DMX/DMXPixelMapping/Source/DMXPixelMappingRuntime/Private/Components/DMXPixelMappingMatrixCellComponent.cpp b/Engine/Plugins/VirtualProduction/DMX/DMXPixelMapping/Source/DMXPixelMappingRuntime/Private/Components/DMXPixelMappingMatrixCellComponent.cpp
--- a/Engine/Plugins/VirtualProduction/DMX/DMXPixelMapping/Source/DMXPixelMappingRuntime/Private/Components/DMXPixelMappingMatrixCellComponent.cpp
+++ b/Engine/Plugins/VirtualProduction/DMX/DMXPixelMapping/Source/DMXPixelMappingRuntime/Private/Components/DMXPixelMappingMatrixCellComponent.cpp
@@ -241,19 +241,18 @@ void UDMXPixelMappingMatrixCellComponent::ResetDMX()
 	// No need to send dmx, that is done by the parent matrix
 }
 
-FString UDMXPixelMappingMatrixCellComponent::GetUserFriendlyName() const
+FString UDMXPixelMappingMatrixCellComponent::GetUserName() const
 {
-	if (UDMXPixelMappingMatrixComponent* MatrixComponent = Cast<UDMXPixelMappingMatrixComponent>(GetParent()))
+	UDMXPixelMappingMatrixComponent* MatrixComponent = Cast<UDMXPixelMappingMatrixComponent>(GetParent());
+	UDMXEntityFixturePatch* FixturePatch = MatrixComponent ? MatrixComponent->FixturePatchRef.GetFixturePatch() : nullptr;
+	if (FixturePatch && UserName.IsEmpty())
 	{
-		UDMXEntityFixturePatch* FixturePatch = MatrixComponent->FixturePatchRef.GetFixturePatch();
-		
-		if (FixturePatch)
-		{
-			return FString::Printf(TEXT("%s: Cell %d"), *FixturePatch->GetDisplayName(), CellID);
-		}
+		return FString::Printf(TEXT("%s: Cell %d"), *FixturePatch->GetDisplayName(), CellID);
+	}
+	else
+	{
+		return FString::Printf(TEXT("%s: Cell %d"), *UserName, CellID);
 	}
-
-	return FString(TEXT("Invalid Patch"));
 }
 
 void UDMXPixelMappingMatrixCellComponent::QueueDownsample()
diff --git a/Engine/Plugins/VirtualProduction/DMX/DMXPixelMapping/Source/DMXPixelMappingRuntime/Private/Components/DMXPixelMappingMatrixComponent.cpp b/Engine/Plugins/VirtualProduction/DMX/DMXPixelMapping/Source/DMXPixelMappingRuntime/Private/Components/DMXPixelMappingMatrixComponent.cpp
--- a/Engine/Plugins/VirtualProduction/DMX/DMXPixelMapping/Source/DMXPixelMappingRuntime/Private/Components/DMXPixelMappingMatrixComponent.cpp
+++ b/Engine/Plugins/VirtualProduction/DMX/DMXPixelMapping/Source/DMXPixelMappingRuntime/Private/Components/DMXPixelMappingMatrixComponent.cpp
@@ -333,7 +333,7 @@ void UDMXPixelMappingMatrixComponent::SendDMX()
 				MapStruct.Map = ColorSpace->GetAttributeNameToValueMap();
 				return MapStruct;
 			});
-		if (!ensureMsgf(AttributeToValueMapArray.Num() == CellComponents.Num(), TEXT("Mismatch num cell attributes and num attribute to value maps. Cannot send DMX for Matrix Component %s"), *GetUserFriendlyName()))
+		if (!ensureMsgf(AttributeToValueMapArray.Num() == CellComponents.Num(), TEXT("Mismatch num cell attributes and num attribute to value maps. Cannot send DMX for Matrix Component %s"), *GetName()))
 		{
 			return;
 		}
@@ -404,14 +404,17 @@ bool UDMXPixelMappingMatrixComponent::CanBeMovedTo(const UDMXPixelMappingBaseCom
 	return false;
 }
 
-FString UDMXPixelMappingMatrixComponent::GetUserFriendlyName() const
+FString UDMXPixelMappingMatrixComponent::GetUserName() const
 {
-	if (UDMXEntityFixturePatch* Patch = FixturePatchRef.GetFixturePatch())
+	const UDMXEntityFixturePatch* Patch = FixturePatchRef.GetFixturePatch();
+	if (Patch && UserName.IsEmpty())
 	{
 		return FString::Printf(TEXT("Fixture Matrix: %s"), *Patch->GetDisplayName());
 	}
-
-	return FString(TEXT("Fixture Matrix: No Fixture Patch"));
+	else
+	{
+		return UserName;
+	}
 }
 
 void UDMXPixelMappingMatrixComponent::SetPosition(const FVector2D& NewPosition)
diff --git a/Engine/Plugins/VirtualProduction/DMX/DMXPixelMapping/Source/DMXPixelMappingRuntime/Private/Components/DMXPixelMappingRendererComponent.cpp b/Engine/Plugins/VirtualProduction/DMX/DMXPixelMapping/Source/DMXPixelMappingRuntime/Private/Components/DMXPixelMappingRendererComponent.cpp
--- a/Engine/Plugins/VirtualProduction/DMX/DMXPixelMapping/Source/DMXPixelMappingRuntime/Private/Components/DMXPixelMappingRendererComponent.cpp
+++ b/Engine/Plugins/VirtualProduction/DMX/DMXPixelMapping/Source/DMXPixelMappingRuntime/Private/Components/DMXPixelMappingRendererComponent.cpp
@@ -280,8 +280,13 @@ void UDMXPixelMappingRendererComponent::RenderAndSendDMX()
 	SendDMX();
 }
 
-FString UDMXPixelMappingRendererComponent::GetUserFriendlyName() const
+FString UDMXPixelMappingRendererComponent::GetUserName() const
 {
+	if (!UserName.IsEmpty())
+	{
+		return UserName;
+	}
+
 	constexpr TCHAR NoSourceString[] = TEXT("None");
 	switch (RendererType)
 	{
diff --git a/Engine/Plugins/VirtualProduction/DMX/DMXPixelMapping/Source/DMXPixelMappingRuntime/Private/Components/DMXPixelMappingRootComponent.cpp b/Engine/Plugins/VirtualProduction/DMX/DMXPixelMapping/Source/DMXPixelMappingRuntime/Private/Components/DMXPixelMappingRootComponent.cpp
--- a/Engine/Plugins/VirtualProduction/DMX/DMXPixelMapping/Source/DMXPixelMappingRuntime/Private/Components/DMXPixelMappingRootComponent.cpp
+++ b/Engine/Plugins/VirtualProduction/DMX/DMXPixelMapping/Source/DMXPixelMappingRuntime/Private/Components/DMXPixelMappingRootComponent.cpp
@@ -101,15 +101,23 @@ UDMXPixelMappingRendererComponent* UDMXPixelMappingRootComponent::FindRendererCo
 	return nullptr;
 }
 
-void UDMXPixelMappingRootComponent::OnComponentRenamed(UDMXPixelMapping* PixelMapping, UDMXPixelMappingBaseComponent* RenamedComponent, UObject* OldOuter, const FName OldName)
+void UDMXPixelMappingRootComponent::OnComponentRenamed(UDMXPixelMappingBaseComponent* RenamedComponent)
 {
+	CachedRendererComponentsByName.Reset();
 	ForEachChildOfClass<UDMXPixelMappingRendererComponent>([this](UDMXPixelMappingRendererComponent* RendererComponent)
 		{
 			CachedRendererComponentsByName.Add(RendererComponent->GetFName(), RendererComponent);
 		}, false);
 }
 
-FString UDMXPixelMappingRootComponent::GetUserFriendlyName() const
+FString UDMXPixelMappingRootComponent::GetUserName() const
 {
-	return GetOuter()->GetName();
+	if (UserName.IsEmpty())
+	{
+		return GetOuter()->GetName();
+	}
+	else
+	{
+		return UserName;
+	}
 }
diff --git a/Engine/Plugins/VirtualProduction/DMX/DMXPixelMapping/Source/DMXPixelMappingRuntime/Private/DMXPixelMapping.cpp b/Engine/Plugins/VirtualProduction/DMX/DMXPixelMapping/Source/DMXPixelMappingRuntime/Private/DMXPixelMapping.cpp
--- a/Engine/Plugins/VirtualProduction/DMX/DMXPixelMapping/Source/DMXPixelMappingRuntime/Private/DMXPixelMapping.cpp
+++ b/Engine/Plugins/VirtualProduction/DMX/DMXPixelMapping/Source/DMXPixelMappingRuntime/Private/DMXPixelMapping.cpp
@@ -156,7 +156,7 @@ void UDMXPixelMapping::RemoveComponent(UDMXPixelMappingBaseComponent* InComponen
 	if (InComponent)
 	{
 #if WITH_EDITOR
-		ensureMsgf(InComponent->GetParent(), TEXT("Trying to remove component %s but it has no valid parent."), *InComponent->GetUserFriendlyName());
+		ensureMsgf(InComponent->GetParent(), TEXT("Trying to remove component %s but it has no valid parent."), *InComponent->GetUserName());
 #endif
 
 		if (UDMXPixelMappingBaseComponent* Parent = InComponent->GetParent())
diff --git a/Engine/Plugins/VirtualProduction/DMX/DMXPixelMapping/Source/DMXPixelMappingRuntime/Public/Components/DMXPixelMappingBaseComponent.h b/Engine/Plugins/VirtualProduction/DMX/DMXPixelMapping/Source/DMXPixelMappingRuntime/Public/Components/DMXPixelMappingBaseComponent.h
--- a/Engine/Plugins/VirtualProduction/DMX/DMXPixelMapping/Source/DMXPixelMappingRuntime/Public/Components/DMXPixelMappingBaseComponent.h
+++ b/Engine/Plugins/VirtualProduction/DMX/DMXPixelMapping/Source/DMXPixelMappingRuntime/Public/Components/DMXPixelMappingBaseComponent.h
@@ -25,7 +25,7 @@ class DMXPIXELMAPPINGRUNTIME_API UDMXPixelMappingBaseComponent
 
 	DECLARE_EVENT_TwoParams(UDMXPixelMappingBaseComponent, FDMXPixelMappingOnComponentAdded, UDMXPixelMapping* /** PixelMapping */, UDMXPixelMappingBaseComponent* /** AddedComponent */);
 	DECLARE_EVENT_TwoParams(UDMXPixelMappingBaseComponent, FDMXPixelMappingOnComponentRemoved, UDMXPixelMapping* /** PixelMapping */, UDMXPixelMappingBaseComponent* /** RemovedComponent */);
-	DECLARE_EVENT_FourParams(UDMXPixelMappingBaseComponent, FDMXPixelMappingOnComponentRenamed, UDMXPixelMapping* /** PixelMapping */, UDMXPixelMappingBaseComponent* /** RenamedComponent */, UObject* /** OldOuter */, const FName /** OldName */);
+	DECLARE_EVENT_OneParam(UDMXPixelMappingBaseComponent, FDMXPixelMappingOnComponentRenamed, UDMXPixelMappingBaseComponent* /** RenamedComponent */);
 
 public:
 	/** Gets an Event broadcast when a component was added */
@@ -34,7 +34,7 @@ public:
 	/** Gets an Event broadcast when a component was added */
 	static FDMXPixelMappingOnComponentRemoved& GetOnComponentRemoved();
 
-	/** Gets an Event broadcast when a component was renamed */
+	/** Gets an Event broadcast this uobject was renamed. Note this is only raised after UObject::Rename, not after SetUserName.*/
 	static FDMXPixelMappingOnComponentRenamed& GetOnComponentRenamed();
 	
 	//~ Begin UObject interface
@@ -94,9 +94,16 @@ public:
 		return false;
 	}
 
-	/** Returns the name of the component used across all widgets that draw it */
+	/** DEPRECATED 5.3 */
+	UE_DEPRECATED(5.3, "Renamed GetUserName to be more explicit aobout the intent. The user name is displayed if set, otherwise the object name.")
 	virtual FString GetUserFriendlyName() const;
 
+	/** Returns the name of the component. If user name is set returns the custom user name (see SetUserName) */
+	virtual FString GetUserName() const;
+
+	/** Sets the custom user name. If set, GetUserName returns this name. */
+	virtual void SetUserName(const FString& NewName);
+
 	/** 
 	 * Loop through all child by given Predicate 
 	 *
@@ -250,6 +257,10 @@ protected:
 	/** Called when the component was added to a parent */
 	virtual void NotifyRemovedFromParent();
 
+	/** Custom user name for the component. Should be used if set. */
+	UPROPERTY()
+	FString UserName;
+
 private:
 	/** Parent component */
 	UPROPERTY(NonTransactional)
diff --git a/Engine/Plugins/VirtualProduction/DMX/DMXPixelMapping/Source/DMXPixelMappingRuntime/Public/Components/DMXPixelMappingFixtureGroupComponent.h b/Engine/Plugins/VirtualProduction/DMX/DMXPixelMapping/Source/DMXPixelMappingRuntime/Public/Components/DMXPixelMappingFixtureGroupComponent.h
--- a/Engine/Plugins/VirtualProduction/DMX/DMXPixelMapping/Source/DMXPixelMappingRuntime/Public/Components/DMXPixelMappingFixtureGroupComponent.h
+++ b/Engine/Plugins/VirtualProduction/DMX/DMXPixelMapping/Source/DMXPixelMappingRuntime/Public/Components/DMXPixelMappingFixtureGroupComponent.h
@@ -40,7 +40,7 @@ public:
 	virtual bool CanBeMovedTo(const UDMXPixelMappingBaseComponent* Component) const override;
 	virtual void ResetDMX() override;
 	virtual void SendDMX() override;
-	virtual FString GetUserFriendlyName() const override;
+	virtual FString GetUserName() const override;
 	//~ End UDMXPixelMappingBaseComponent implementation
 
 	//~ Begin UDMXPixelMappingOutputComponent implementation
diff --git a/Engine/Plugins/VirtualProduction/DMX/DMXPixelMapping/Source/DMXPixelMappingRuntime/Public/Components/DMXPixelMappingFixtureGroupItemComponent.h b/Engine/Plugins/VirtualProduction/DMX/DMXPixelMapping/Source/DMXPixelMappingRuntime/Public/Components/DMXPixelMappingFixtureGroupItemComponent.h
--- a/Engine/Plugins/VirtualProduction/DMX/DMXPixelMapping/Source/DMXPixelMappingRuntime/Public/Components/DMXPixelMappingFixtureGroupItemComponent.h
+++ b/Engine/Plugins/VirtualProduction/DMX/DMXPixelMapping/Source/DMXPixelMappingRuntime/Public/Components/DMXPixelMappingFixtureGroupItemComponent.h
@@ -51,7 +51,7 @@ public:
 	virtual const FName& GetNamePrefix() override;
 	virtual void ResetDMX() override;
 	virtual void SendDMX() override;
-	virtual FString GetUserFriendlyName() const override;
+	virtual FString GetUserName() const override;
 	//~ End UDMXPixelMappingBaseComponent implementation
 
 	//~ Begin UDMXPixelMappingOutputComponent implementation
diff --git a/Engine/Plugins/VirtualProduction/DMX/DMXPixelMapping/Source/DMXPixelMappingRuntime/Public/Components/DMXPixelMappingMatrixCellComponent.h b/Engine/Plugins/VirtualProduction/DMX/DMXPixelMapping/Source/DMXPixelMappingRuntime/Public/Components/DMXPixelMappingMatrixCellComponent.h
--- a/Engine/Plugins/VirtualProduction/DMX/DMXPixelMapping/Source/DMXPixelMappingRuntime/Public/Components/DMXPixelMappingMatrixCellComponent.h
+++ b/Engine/Plugins/VirtualProduction/DMX/DMXPixelMapping/Source/DMXPixelMappingRuntime/Public/Components/DMXPixelMappingMatrixCellComponent.h
@@ -46,7 +46,7 @@ public:
 	//~ Begin UDMXPixelMappingBaseComponent implementation
 	virtual const FName& GetNamePrefix() override;
 	virtual void ResetDMX() override;
-	virtual FString GetUserFriendlyName() const override;
+	virtual FString GetUserName() const override;
 	//~ End UDMXPixelMappingBaseComponent implementation
 
 	//~ Begin UDMXPixelMappingOutputComponent implementation
diff --git a/Engine/Plugins/VirtualProduction/DMX/DMXPixelMapping/Source/DMXPixelMappingRuntime/Public/Components/DMXPixelMappingMatrixComponent.h b/Engine/Plugins/VirtualProduction/DMX/DMXPixelMapping/Source/DMXPixelMappingRuntime/Public/Components/DMXPixelMappingMatrixComponent.h
--- a/Engine/Plugins/VirtualProduction/DMX/DMXPixelMapping/Source/DMXPixelMappingRuntime/Public/Components/DMXPixelMappingMatrixComponent.h
+++ b/Engine/Plugins/VirtualProduction/DMX/DMXPixelMapping/Source/DMXPixelMappingRuntime/Public/Components/DMXPixelMappingMatrixComponent.h
@@ -60,7 +60,7 @@ public:
 	virtual void ResetDMX() override;
 	virtual void SendDMX() override;
 	virtual bool CanBeMovedTo(const UDMXPixelMappingBaseComponent* Component) const override;
-	virtual FString GetUserFriendlyName() const override;
+	virtual FString GetUserName() const override;
 	// ~End UDMXPixelMappingBaseComponent interface
 
 	// ~Begin UDMXPixelMappingOutputComponent interface
diff --git a/Engine/Plugins/VirtualProduction/DMX/DMXPixelMapping/Source/DMXPixelMappingRuntime/Public/Components/DMXPixelMappingRendererComponent.h b/Engine/Plugins/VirtualProduction/DMX/DMXPixelMapping/Source/DMXPixelMappingRuntime/Public/Components/DMXPixelMappingRendererComponent.h
--- a/Engine/Plugins/VirtualProduction/DMX/DMXPixelMapping/Source/DMXPixelMappingRuntime/Public/Components/DMXPixelMappingRendererComponent.h
+++ b/Engine/Plugins/VirtualProduction/DMX/DMXPixelMapping/Source/DMXPixelMappingRuntime/Public/Components/DMXPixelMappingRendererComponent.h
@@ -53,7 +53,7 @@ public:
 	virtual void SendDMX() override;
 	virtual void Render() final;
 	virtual void RenderAndSendDMX() final;
-	virtual FString GetUserFriendlyName() const override;
+	virtual FString GetUserName() const override;
 	//~ End UDMXPixelMappingBaseComponent implementation
 
 	/** Returns the pixel mapping preprocess renderer */
diff --git a/Engine/Plugins/VirtualProduction/DMX/DMXPixelMapping/Source/DMXPixelMappingRuntime/Public/Components/DMXPixelMappingRootComponent.h b/Engine/Plugins/VirtualProduction/DMX/DMXPixelMapping/Source/DMXPixelMappingRuntime/Public/Components/DMXPixelMappingRootComponent.h
--- a/Engine/Plugins/VirtualProduction/DMX/DMXPixelMapping/Source/DMXPixelMappingRuntime/Public/Components/DMXPixelMappingRootComponent.h
+++ b/Engine/Plugins/VirtualProduction/DMX/DMXPixelMapping/Source/DMXPixelMappingRuntime/Public/Components/DMXPixelMappingRootComponent.h
@@ -33,7 +33,7 @@ public:
 	virtual void SendDMX() override;
 	virtual void Render() override;
 	virtual void RenderAndSendDMX() override;
-	virtual FString GetUserFriendlyName() const override;
+	virtual FString GetUserName() const override;
 	//~ End UDMXPixelMappingBaseComponent implementation
 
 	/** Check if a Component can be moved under another one (used for copy/move/duplicate) */
@@ -44,7 +44,7 @@ public:
 
 private:
 	/** Called when a component in the PixelMapping was renamed */
-	void OnComponentRenamed(UDMXPixelMapping* PixelMapping , UDMXPixelMappingBaseComponent* RenamedComponent, UObject* OldOuter, const FName OldName);
+	void OnComponentRenamed(UDMXPixelMappingBaseComponent* RenamedComponent);
 
 	/** Chached Renderer Component names and their actual object pointer. Useful to speed up access (see DMXPixelMappingSubsystem) */
 	UPROPERTY(Transient)
