DMX - Remove type aliases to improve readability of pixel mapping code, reduce inclusion of monolithic headers

diff --git a/Engine/Plugins/VirtualProduction/DMX/DMXPixelMapping/Source/DMXPixelMappingBlueprintGraph/Public/K2Node_PixelMappingBaseComponent.h b/Engine/Plugins/VirtualProduction/DMX/DMXPixelMapping/Source/DMXPixelMappingBlueprintGraph/Public/K2Node_PixelMappingBaseComponent.h
--- a/Engine/Plugins/VirtualProduction/DMX/DMXPixelMapping/Source/DMXPixelMappingBlueprintGraph/Public/K2Node_PixelMappingBaseComponent.h
+++ b/Engine/Plugins/VirtualProduction/DMX/DMXPixelMapping/Source/DMXPixelMappingBlueprintGraph/Public/K2Node_PixelMappingBaseComponent.h
@@ -2,12 +2,13 @@
 
 #pragma once
 
-#include "CoreMinimal.h"
 #include "K2Node.h"
+
 #include "K2Node_PixelMappingBaseComponent.generated.h"
 
-class UDMXPixelMapping;
 class FKismetCompilerContext;
+class UDMXPixelMapping;
+
 
 /**
  * Base Pixel Mapping node. Never use directly
diff --git a/Engine/Plugins/VirtualProduction/DMX/DMXPixelMapping/Source/DMXPixelMappingBlueprintGraph/Public/K2Node_PixelMappingRendererComponent.h b/Engine/Plugins/VirtualProduction/DMX/DMXPixelMapping/Source/DMXPixelMappingBlueprintGraph/Public/K2Node_PixelMappingRendererComponent.h
--- a/Engine/Plugins/VirtualProduction/DMX/DMXPixelMapping/Source/DMXPixelMappingBlueprintGraph/Public/K2Node_PixelMappingRendererComponent.h
+++ b/Engine/Plugins/VirtualProduction/DMX/DMXPixelMapping/Source/DMXPixelMappingBlueprintGraph/Public/K2Node_PixelMappingRendererComponent.h
@@ -2,8 +2,8 @@
 
 #pragma once
 
-#include "CoreMinimal.h"
 #include "K2Node_PixelMappingBaseComponent.h"
+
 #include "K2Node_PixelMappingRendererComponent.generated.h"
 
 /**
diff --git a/Engine/Plugins/VirtualProduction/DMX/DMXPixelMapping/Source/DMXPixelMappingCore/Public/DMXPixelMappingTypes.h b/Engine/Plugins/VirtualProduction/DMX/DMXPixelMapping/Source/DMXPixelMappingCore/Public/DMXPixelMappingTypes.h
--- a/Engine/Plugins/VirtualProduction/DMX/DMXPixelMapping/Source/DMXPixelMappingCore/Public/DMXPixelMappingTypes.h
+++ b/Engine/Plugins/VirtualProduction/DMX/DMXPixelMapping/Source/DMXPixelMappingCore/Public/DMXPixelMappingTypes.h
@@ -2,7 +2,8 @@
 
 #pragma once
 
-#include "CoreMinimal.h"
+#include "HAL/Platform.h"
+
 #include "DMXPixelMappingTypes.generated.h"
 
 
@@ -44,4 +45,4 @@ enum class EDMXColorMode : uint8
 {
 	CM_RGB UMETA(DisplayName = "RGB"),
 	CM_Monochrome UMETA(DisplayName = "Monochrome"),
-};
\ No newline at end of file
+};
diff --git a/Engine/Plugins/VirtualProduction/DMX/DMXPixelMapping/Source/DMXPixelMappingCore/Public/DMXPixelMappingUtils.h b/Engine/Plugins/VirtualProduction/DMX/DMXPixelMapping/Source/DMXPixelMappingCore/Public/DMXPixelMappingUtils.h
--- a/Engine/Plugins/VirtualProduction/DMX/DMXPixelMapping/Source/DMXPixelMappingCore/Public/DMXPixelMappingUtils.h
+++ b/Engine/Plugins/VirtualProduction/DMX/DMXPixelMapping/Source/DMXPixelMappingCore/Public/DMXPixelMappingUtils.h
@@ -2,7 +2,7 @@
 
 #pragma once
 
-#include "CoreMinimal.h"
+#include "HAL/Platform.h"
 
 #include "Library/DMXEntityFixtureType.h"
 
diff --git a/Engine/Plugins/VirtualProduction/DMX/DMXPixelMapping/Source/DMXPixelMappingEditor/Private/AssetsTools/AssetTypeActions_DMXPixelMapping.cpp b/Engine/Plugins/VirtualProduction/DMX/DMXPixelMapping/Source/DMXPixelMappingEditor/Private/AssetsTools/AssetTypeActions_DMXPixelMapping.cpp
--- a/Engine/Plugins/VirtualProduction/DMX/DMXPixelMapping/Source/DMXPixelMappingEditor/Private/AssetsTools/AssetTypeActions_DMXPixelMapping.cpp
+++ b/Engine/Plugins/VirtualProduction/DMX/DMXPixelMapping/Source/DMXPixelMappingEditor/Private/AssetsTools/AssetTypeActions_DMXPixelMapping.cpp
@@ -2,7 +2,7 @@
 
 #include "AssetsTools/AssetTypeActions_DMXPixelMapping.h"
 #include "DMXPixelMapping.h"
-#include "DMXPixelMappingEditorCommon.h"
+#include "DMXPixelMappingEditorLog.h"
 #include "DMXPixelMappingEditorModule.h"
 #include "Toolkits/DMXPixelMappingToolkit.h"
 
diff --git a/Engine/Plugins/VirtualProduction/DMX/DMXPixelMapping/Source/DMXPixelMappingEditor/Private/DMXPixelMappingEditorCommon.cpp b/Engine/Plugins/VirtualProduction/DMX/DMXPixelMapping/Source/DMXPixelMappingEditor/Private/DMXPixelMappingEditorCommon.cpp
deleted file mode 100644
diff --git a/Engine/Plugins/VirtualProduction/DMX/DMXPixelMapping/Source/DMXPixelMappingEditor/Private/DMXPixelMappingEditorCommon.h b/Engine/Plugins/VirtualProduction/DMX/DMXPixelMapping/Source/DMXPixelMappingEditor/Private/DMXPixelMappingEditorCommon.h
deleted file mode 100644
diff --git a/Engine/Plugins/VirtualProduction/DMX/DMXPixelMapping/Source/DMXPixelMappingEditor/Private/DMXPixelMappingEditorLog.cpp b/Engine/Plugins/VirtualProduction/DMX/DMXPixelMapping/Source/DMXPixelMappingEditor/Private/DMXPixelMappingEditorLog.cpp
new file mode 100644
--- /dev/null
+++ b/Engine/Plugins/VirtualProduction/DMX/DMXPixelMapping/Source/DMXPixelMappingEditor/Private/DMXPixelMappingEditorLog.cpp
@@ -0,0 +1,5 @@
+// Copyright Epic Games, Inc. All Rights Reserved.
+
+#include "DMXPixelMappingEditorLog.h"
+
+DEFINE_LOG_CATEGORY(LogDMXPixelMappingEditor)
\ No newline at end of file
diff --git a/Engine/Plugins/VirtualProduction/DMX/DMXPixelMapping/Source/DMXPixelMappingEditor/Private/DMXPixelMappingEditorLog.h b/Engine/Plugins/VirtualProduction/DMX/DMXPixelMapping/Source/DMXPixelMappingEditor/Private/DMXPixelMappingEditorLog.h
new file mode 100644
--- /dev/null
+++ b/Engine/Plugins/VirtualProduction/DMX/DMXPixelMapping/Source/DMXPixelMappingEditor/Private/DMXPixelMappingEditorLog.h
@@ -0,0 +1,9 @@
+// Copyright Epic Games, Inc. All Rights Reserved.
+
+#pragma once
+
+#include "Logging/LogMacros.h"
+
+
+DECLARE_LOG_CATEGORY_EXTERN(LogDMXPixelMappingEditor, All, All);
+
diff --git a/Engine/Plugins/VirtualProduction/DMX/DMXPixelMapping/Source/DMXPixelMappingEditor/Private/DMXPixelMappingEditorUtils.cpp b/Engine/Plugins/VirtualProduction/DMX/DMXPixelMapping/Source/DMXPixelMappingEditor/Private/DMXPixelMappingEditorUtils.cpp
--- a/Engine/Plugins/VirtualProduction/DMX/DMXPixelMapping/Source/DMXPixelMappingEditor/Private/DMXPixelMappingEditorUtils.cpp
+++ b/Engine/Plugins/VirtualProduction/DMX/DMXPixelMapping/Source/DMXPixelMappingEditor/Private/DMXPixelMappingEditorUtils.cpp
@@ -3,21 +3,19 @@
 #include "DMXPixelMappingEditorUtils.h"
 
 #include "DMXPixelMapping.h"
-#include "DMXPixelMappingEditorCommon.h"
+#include "DMXPixelMappingEditorLog.h"
 #include "Components/DMXPixelMappingRendererComponent.h"
 #include "Components/DMXPixelMappingRootComponent.h"
 #include "DragDrop/DMXPixelMappingDragDropOp.h"
 #include "Framework/Application/SlateApplication.h"
+#include "Framework/Commands/GenericCommands.h"
+#include "Framework/MultiBox/MultiBoxBuilder.h"
 #include "Layout/WidgetPath.h"
 #include "Toolkits/DMXPixelMappingToolkit.h"
 
 
-#include "Framework/Commands/GenericCommands.h"
-#include "Framework/MultiBox/MultiBoxBuilder.h"
-
 #define LOCTEXT_NAMESPACE "FDMXPixelMappingEditorUtils"
 
-
 UDMXPixelMappingRendererComponent* FDMXPixelMappingEditorUtils::AddRenderer(UDMXPixelMapping* InPixelMapping)
 {
 	if (InPixelMapping == nullptr)
diff --git a/Engine/Plugins/VirtualProduction/DMX/DMXPixelMapping/Source/DMXPixelMappingEditor/Private/ViewModels/DMXPixelMappingHierarchyItem.h b/Engine/Plugins/VirtualProduction/DMX/DMXPixelMapping/Source/DMXPixelMappingEditor/Private/ViewModels/DMXPixelMappingHierarchyItem.h
--- a/Engine/Plugins/VirtualProduction/DMX/DMXPixelMapping/Source/DMXPixelMappingEditor/Private/ViewModels/DMXPixelMappingHierarchyItem.h
+++ b/Engine/Plugins/VirtualProduction/DMX/DMXPixelMapping/Source/DMXPixelMappingEditor/Private/ViewModels/DMXPixelMappingHierarchyItem.h
@@ -3,7 +3,6 @@
 #pragma once
 
 #include "Algo/Sort.h"
-#include "DMXPixelMappingEditorCommon.h"
 #include "DMXPixelMappingComponentReference.h"
 
 struct FSlateColor;
@@ -98,5 +97,5 @@ private:
 	TWeakObjectPtr<UDMXPixelMappingBaseComponent> WeakComponent;
 
 	/** Toolkit this model resides in */
-	FDMXPixelMappingToolkitWeakPtr WeakToolkit;
+	TWeakPtr<FDMXPixelMappingToolkit> WeakToolkit;
 };
diff --git a/Engine/Plugins/VirtualProduction/DMX/DMXPixelMapping/Source/DMXPixelMappingEditor/Private/ViewModels/DMXPixelMappingLayoutViewModel.cpp b/Engine/Plugins/VirtualProduction/DMX/DMXPixelMapping/Source/DMXPixelMappingEditor/Private/ViewModels/DMXPixelMappingLayoutViewModel.cpp
--- a/Engine/Plugins/VirtualProduction/DMX/DMXPixelMapping/Source/DMXPixelMappingEditor/Private/ViewModels/DMXPixelMappingLayoutViewModel.cpp
+++ b/Engine/Plugins/VirtualProduction/DMX/DMXPixelMapping/Source/DMXPixelMappingEditor/Private/ViewModels/DMXPixelMappingLayoutViewModel.cpp
@@ -3,12 +3,12 @@
 #include "ViewModels/DMXPixelMappingLayoutViewModel.h"
 
 #include "Components/DMXPixelMappingMatrixCellComponent.h"
-#include "DMXPixelMappingEditorCommon.h"
 #include "Components/DMXPixelMappingFixtureGroupComponent.h"
 #include "Components/DMXPixelMappingFixtureGroupItemComponent.h"
 #include "Components/DMXPixelMappingMatrixComponent.h"
 #include "Components/DMXPixelMappingRendererComponent.h"
 #include "Components/DMXPixelMappingScreenComponent.h"
+#include "DMXPixelMappingEditorLog.h"
 #include "LayoutScripts/DMXPixelMappingLayoutScript.h"
 #include "ScopedTransaction.h"
 #include "Settings/DMXPixelMappingEditorSettings.h"
diff --git a/Engine/Plugins/VirtualProduction/DMX/DMXPixelMapping/Source/DMXPixelMappingEditor/Private/Views/SDMXPixelMappingHierarchyView.cpp b/Engine/Plugins/VirtualProduction/DMX/DMXPixelMapping/Source/DMXPixelMappingEditor/Private/Views/SDMXPixelMappingHierarchyView.cpp
--- a/Engine/Plugins/VirtualProduction/DMX/DMXPixelMapping/Source/DMXPixelMappingEditor/Private/Views/SDMXPixelMappingHierarchyView.cpp
+++ b/Engine/Plugins/VirtualProduction/DMX/DMXPixelMapping/Source/DMXPixelMappingEditor/Private/Views/SDMXPixelMappingHierarchyView.cpp
@@ -69,7 +69,7 @@ namespace UE::DMX::PixelMappingEditor::SDMXPixelMappingHierarchyView::Private
 	};
 
 	/** Helper to adopt the expansion of items in theh hierarchy treeview */
-	void AdoptItemExpansionFromTree(const TSharedRef<STreeView<FDMXPixelMappingHierarchyItemWidgetModelPtr>>& TreeView, const TSharedPtr<FDMXPixelMappingHierarchyItem>& Parent)
+	void AdoptItemExpansionFromTree(const TSharedRef<STreeView<TSharedPtr<FDMXPixelMappingHierarchyItem>>>& TreeView, const TSharedPtr<FDMXPixelMappingHierarchyItem>& Parent)
 	{
 		TreeView->SetItemExpansion(Parent, Parent->IsExpanded());
 		for (const TSharedPtr<FDMXPixelMappingHierarchyItem>& Child : Parent->GetChildren())
@@ -102,7 +102,7 @@ void SDMXPixelMappingHierarchyView::Construct(const FArguments& InArgs, const TS
 		FExecuteAction::CreateSP(this, &SDMXPixelMappingHierarchyView::DeleteSelectedComponents)
 	);
 
-	using TextFilterType = TTextFilter<FDMXPixelMappingHierarchyItemWidgetModelPtr>;
+	using TextFilterType = TTextFilter<TSharedPtr<FDMXPixelMappingHierarchyItem>>;
 	SearchFilter = MakeShared<TextFilterType>(TextFilterType::FItemToStringArray::CreateSP(this, &SDMXPixelMappingHierarchyView::GetWidgetFilterStrings));
 
 	FilterHandler = MakeShared<TreeFilterHandler<TSharedPtr<FDMXPixelMappingHierarchyItem>>>();
@@ -190,7 +190,7 @@ void SDMXPixelMappingHierarchyView::BuildChildSlotAndRefresh()
 			.FillHeight(1.f)
 			.Padding(4.f)
 			[
-				SAssignNew(HierarchyTreeView, STreeView<FDMXPixelMappingHierarchyItemWidgetModelPtr>)
+				SAssignNew(HierarchyTreeView, STreeView<TSharedPtr<FDMXPixelMappingHierarchyItem>>)
 				.ItemHeight(20.0f)
 				.SelectionMode(ESelectionMode::Multi)
 				.HeaderRow(GenerateHeaderRow())
@@ -200,7 +200,10 @@ void SDMXPixelMappingHierarchyView::BuildChildSlotAndRefresh()
 				.OnExpansionChanged(this, &SDMXPixelMappingHierarchyView::OnHierarchyExpansionChanged)
 				.OnContextMenuOpening(this, &SDMXPixelMappingHierarchyView::OnContextMenuOpening)
 				.TreeItemsSource(&FilteredRootItems)
-				.OnItemToString_Debug_Lambda([this](FDMXPixelMappingHierarchyItemWidgetModelPtr Item) { return Item->GetComponentNameText().ToString(); })
+				.OnItemToString_Debug_Lambda([this](TSharedPtr<FDMXPixelMappingHierarchyItem> Item) 
+					{ 
+						return Item->GetComponentNameText().ToString(); 
+					})
 			]
 		];
 
@@ -380,7 +383,7 @@ void SDMXPixelMappingHierarchyView::OnGetChildItems(TSharedPtr<FDMXPixelMappingH
 	OutChildren = InParent->GetChildren();
 }
 
-TSharedRef<ITableRow> SDMXPixelMappingHierarchyView::OnGenerateRow(FDMXPixelMappingHierarchyItemWidgetModelPtr Item, const TSharedRef<STableViewBase>& OwnerTable)
+TSharedRef<ITableRow> SDMXPixelMappingHierarchyView::OnGenerateRow(TSharedPtr<FDMXPixelMappingHierarchyItem> Item, const TSharedRef<STableViewBase>& OwnerTable)
 {
 	return SNew(SDMXPixelMappingHierarchyRow, OwnerTable, WeakToolkit, Item.ToSharedRef());
 }
@@ -404,7 +407,7 @@ TSharedPtr<SWidget> SDMXPixelMappingHierarchyView::OnContextMenuOpening()
 	return MenuBuilder.MakeWidget();
 }
 
-void SDMXPixelMappingHierarchyView::OnHierarchySelectionChanged(FDMXPixelMappingHierarchyItemWidgetModelPtr SelectedItem, ESelectInfo::Type SelectInfo)
+void SDMXPixelMappingHierarchyView::OnHierarchySelectionChanged(TSharedPtr<FDMXPixelMappingHierarchyItem> SelectedItem, ESelectInfo::Type SelectInfo)
 {
 	if (SelectInfo == ESelectInfo::Direct)
 	{
@@ -416,8 +419,8 @@ void SDMXPixelMappingHierarchyView::OnHierarchySelectionChanged(FDMXPixelMapping
 	if (TSharedPtr<FDMXPixelMappingToolkit> Toolkit = WeakToolkit.Pin())
 	{
 		TSet<FDMXPixelMappingComponentReference> ComponentsToSelect;
-		FDMXPixelMappingHierarchyItemWidgetModelArr SelectedItems = HierarchyTreeView->GetSelectedItems();
-		for (FDMXPixelMappingHierarchyItemWidgetModelPtr& Item : SelectedItems)
+		const TArray<TSharedPtr<FDMXPixelMappingHierarchyItem>> SelectedItems = HierarchyTreeView->GetSelectedItems();
+		for (const TSharedPtr<FDMXPixelMappingHierarchyItem>& Item : SelectedItems)
 		{
 			ComponentsToSelect.Add(FDMXPixelMappingComponentReference(Toolkit, Item->GetComponent()));
 		}
@@ -592,7 +595,7 @@ void SDMXPixelMappingHierarchyView::SetFilterText(const FText& Text)
 	FilterHandler->RefreshAndFilterTree();
 }
 
-void SDMXPixelMappingHierarchyView::GetWidgetFilterStrings(FDMXPixelMappingHierarchyItemWidgetModelPtr InModel, TArray<FString>& OutStrings) const
+void SDMXPixelMappingHierarchyView::GetWidgetFilterStrings(TSharedPtr<FDMXPixelMappingHierarchyItem> InModel, TArray<FString>& OutStrings) const
 {
 	OutStrings.Add(InModel->GetComponentNameText().ToString());
 	OutStrings.Add(InModel->GetFixtureIDText().ToString());
diff --git a/Engine/Plugins/VirtualProduction/DMX/DMXPixelMapping/Source/DMXPixelMappingEditor/Private/Views/SDMXPixelMappingHierarchyView.h b/Engine/Plugins/VirtualProduction/DMX/DMXPixelMapping/Source/DMXPixelMappingEditor/Private/Views/SDMXPixelMappingHierarchyView.h
--- a/Engine/Plugins/VirtualProduction/DMX/DMXPixelMapping/Source/DMXPixelMappingEditor/Private/Views/SDMXPixelMappingHierarchyView.h
+++ b/Engine/Plugins/VirtualProduction/DMX/DMXPixelMapping/Source/DMXPixelMappingEditor/Private/Views/SDMXPixelMappingHierarchyView.h
@@ -2,7 +2,7 @@
 
 #pragma once
 
-#include "DMXPixelMappingEditorCommon.h"
+#include "DMXPixelMappingEditorLog.h"
 #include "EditorUndoClient.h"
 #include "Engine/EngineTypes.h"
 #include "Widgets/SCompoundWidget.h"
@@ -12,7 +12,6 @@ class FDMXPixelMappingHierarchyItem;
 class FDMXPixelMappingToolkit;
 class FUICommandList;
 class ITableRow;
-class SBorder;
 class SHeaderRow;
 class SSearchBox;
 class STableViewBase;
@@ -74,7 +73,7 @@ private:
 	//~ End FEditorUndoClient interface
 
 	/** Called to get child items */
-	void OnGetChildItems(TSharedPtr<FDMXPixelMappingHierarchyItem> InParent, FDMXPixelMappingHierarchyItemWidgetModelArr& OutChildren);
+	void OnGetChildItems(TSharedPtr<FDMXPixelMappingHierarchyItem> InParent, TArray<TSharedPtr<FDMXPixelMappingHierarchyItem>>& OutChildren);
 
 	/** Called when a row is generated */
 	TSharedRef<ITableRow> OnGenerateRow(TSharedPtr<FDMXPixelMappingHierarchyItem> Item, const TSharedRef<STableViewBase>& OwnerTable);
diff --git a/Engine/Plugins/VirtualProduction/DMX/DMXPixelMapping/Source/DMXPixelMappingEditor/Private/Widgets/SDMXPixelMappingHierarchyRow.cpp b/Engine/Plugins/VirtualProduction/DMX/DMXPixelMapping/Source/DMXPixelMappingEditor/Private/Widgets/SDMXPixelMappingHierarchyRow.cpp
--- a/Engine/Plugins/VirtualProduction/DMX/DMXPixelMapping/Source/DMXPixelMappingEditor/Private/Widgets/SDMXPixelMappingHierarchyRow.cpp
+++ b/Engine/Plugins/VirtualProduction/DMX/DMXPixelMapping/Source/DMXPixelMappingEditor/Private/Widgets/SDMXPixelMappingHierarchyRow.cpp
@@ -22,7 +22,7 @@ void SDMXPixelMappingHierarchyRow::Construct(const FArguments& InArgs, const TSh
 	WeakToolkit = InWeakToolkit;
 	Item = InItem;
 
-	SMultiColumnTableRow<FDMXPixelMappingHierarchyItemWidgetModelPtr>::Construct(
+	SMultiColumnTableRow<TSharedPtr<FDMXPixelMappingHierarchyItem>>::Construct(
 		FSuperRowType::FArguments()
 		.Padding(0.0f)
 		.Style(&FAppStyle::Get().GetWidgetStyle<FTableRowStyle>("SimpleTableView.Row"))
diff --git a/Engine/Plugins/VirtualProduction/DMX/DMXPixelMapping/Source/DMXPixelMappingEditor/Private/Widgets/SDMXPixelMappingHierarchyRow.h b/Engine/Plugins/VirtualProduction/DMX/DMXPixelMapping/Source/DMXPixelMappingEditor/Private/Widgets/SDMXPixelMappingHierarchyRow.h
--- a/Engine/Plugins/VirtualProduction/DMX/DMXPixelMapping/Source/DMXPixelMappingEditor/Private/Widgets/SDMXPixelMappingHierarchyRow.h
+++ b/Engine/Plugins/VirtualProduction/DMX/DMXPixelMapping/Source/DMXPixelMappingEditor/Private/Widgets/SDMXPixelMappingHierarchyRow.h
@@ -2,11 +2,11 @@
 
 #pragma once
 
-#include "DMXPixelMappingEditorCommon.h"
 #include "Widgets/Views/ITableRow.h"
 #include "Widgets/Views/STableRow.h"
 
 class FDMXPixelMappingHierarchyItem;
+class FDMXPixelMappingToolkit;
 class SInlineEditableTextBlock;
 class SDMXPixelMappingHierarchyView;
 class SImage;
diff --git a/Engine/Plugins/VirtualProduction/DMX/DMXPixelMapping/Source/DMXPixelMappingRenderer/Public/IDMXPixelMappingRenderer.h b/Engine/Plugins/VirtualProduction/DMX/DMXPixelMapping/Source/DMXPixelMappingRenderer/Public/IDMXPixelMappingRenderer.h
--- a/Engine/Plugins/VirtualProduction/DMX/DMXPixelMapping/Source/DMXPixelMappingRenderer/Public/IDMXPixelMappingRenderer.h
+++ b/Engine/Plugins/VirtualProduction/DMX/DMXPixelMapping/Source/DMXPixelMappingRenderer/Public/IDMXPixelMappingRenderer.h
@@ -2,7 +2,6 @@
 
 #pragma once
 
-#include "CoreMinimal.h"
 #include "RHI.h"
 #include "RHIResources.h"
 
diff --git a/Engine/Plugins/VirtualProduction/DMX/DMXPixelMapping/Source/DMXPixelMappingRuntime/Private/Components/DMXPixelMappingBaseComponent.cpp b/Engine/Plugins/VirtualProduction/DMX/DMXPixelMapping/Source/DMXPixelMappingRuntime/Private/Components/DMXPixelMappingBaseComponent.cpp
--- a/Engine/Plugins/VirtualProduction/DMX/DMXPixelMapping/Source/DMXPixelMappingRuntime/Private/Components/DMXPixelMappingBaseComponent.cpp
+++ b/Engine/Plugins/VirtualProduction/DMX/DMXPixelMapping/Source/DMXPixelMappingRuntime/Private/Components/DMXPixelMappingBaseComponent.cpp
@@ -3,7 +3,6 @@
 #include "Components/DMXPixelMappingBaseComponent.h"
 
 #include "DMXPixelMapping.h"
-#include "DMXPixelMappingRuntimeCommon.h"
 #include "DMXPixelMappingMainStreamObjectVersion.h"
 #include "Components/DMXPixelMappingRendererComponent.h"
 #include "Components/DMXPixelMappingRootComponent.h"
diff --git a/Engine/Plugins/VirtualProduction/DMX/DMXPixelMapping/Source/DMXPixelMappingRuntime/Private/Components/DMXPixelMappingMatrixComponent.cpp b/Engine/Plugins/VirtualProduction/DMX/DMXPixelMapping/Source/DMXPixelMappingRuntime/Private/Components/DMXPixelMappingMatrixComponent.cpp
--- a/Engine/Plugins/VirtualProduction/DMX/DMXPixelMapping/Source/DMXPixelMappingRuntime/Private/Components/DMXPixelMappingMatrixComponent.cpp
+++ b/Engine/Plugins/VirtualProduction/DMX/DMXPixelMapping/Source/DMXPixelMappingRuntime/Private/Components/DMXPixelMappingMatrixComponent.cpp
@@ -9,7 +9,7 @@
 #include "DMXConversions.h"
 #include "DMXPixelMapping.h"
 #include "DMXPixelMappingMainStreamObjectVersion.h"
-#include "DMXPixelMappingRuntimeCommon.h"
+#include "DMXPixelMappingRuntimeLog.h"
 #include "DMXPixelMappingRuntimeUtils.h"
 #include "DMXPixelMappingTypes.h"
 #include "ColorSpace/DMXPixelMappingColorSpace_RGBCMY.h"
@@ -29,7 +29,6 @@
 
 #define LOCTEXT_NAMESPACE "DMXPixelMappingMatrixComponent"
 
-
 UDMXPixelMappingMatrixComponent::UDMXPixelMappingMatrixComponent()
 {
 	ColorSpaceClass = UDMXPixelMappingColorSpace_RGBCMY::StaticClass();
diff --git a/Engine/Plugins/VirtualProduction/DMX/DMXPixelMapping/Source/DMXPixelMappingRuntime/Private/Components/DMXPixelMappingScreenComponent.cpp b/Engine/Plugins/VirtualProduction/DMX/DMXPixelMapping/Source/DMXPixelMappingRuntime/Private/Components/DMXPixelMappingScreenComponent.cpp
--- a/Engine/Plugins/VirtualProduction/DMX/DMXPixelMapping/Source/DMXPixelMappingRuntime/Private/Components/DMXPixelMappingScreenComponent.cpp
+++ b/Engine/Plugins/VirtualProduction/DMX/DMXPixelMapping/Source/DMXPixelMappingRuntime/Private/Components/DMXPixelMappingScreenComponent.cpp
@@ -2,6 +2,8 @@
 
 #include "Components/DMXPixelMappingScreenComponent.h"
 
+#include "DMXPixelMappingRuntimeLog.h"
+#include "DMXPixelMappingTypes.h"
 #include "DMXPixelMappingUtils.h"
 #include "Components/DMXPixelMappingRendererComponent.h"
 #include "IO/DMXOutputPort.h"
diff --git a/Engine/Plugins/VirtualProduction/DMX/DMXPixelMapping/Source/DMXPixelMappingRuntime/Private/DMXPixelMappingRuntimeCommon.cpp b/Engine/Plugins/VirtualProduction/DMX/DMXPixelMapping/Source/DMXPixelMappingRuntime/Private/DMXPixelMappingRuntimeCommon.cpp
deleted file mode 100644
diff --git a/Engine/Plugins/VirtualProduction/DMX/DMXPixelMapping/Source/DMXPixelMappingRuntime/Private/DMXPixelMappingRuntimeLog.cpp b/Engine/Plugins/VirtualProduction/DMX/DMXPixelMapping/Source/DMXPixelMappingRuntime/Private/DMXPixelMappingRuntimeLog.cpp
new file mode 100644
--- /dev/null
+++ b/Engine/Plugins/VirtualProduction/DMX/DMXPixelMapping/Source/DMXPixelMappingRuntime/Private/DMXPixelMappingRuntimeLog.cpp
@@ -0,0 +1,5 @@
+// Copyright Epic Games, Inc. All Rights Reserved.
+
+#include "DMXPixelMappingRuntimeLog.h"
+
+DEFINE_LOG_CATEGORY(LogDMXPixelMappingRuntime)
\ No newline at end of file
diff --git a/Engine/Plugins/VirtualProduction/DMX/DMXPixelMapping/Source/DMXPixelMappingRuntime/Public/Blueprint/DMXPixelMappingSubsystem.h b/Engine/Plugins/VirtualProduction/DMX/DMXPixelMapping/Source/DMXPixelMappingRuntime/Public/Blueprint/DMXPixelMappingSubsystem.h
--- a/Engine/Plugins/VirtualProduction/DMX/DMXPixelMapping/Source/DMXPixelMappingRuntime/Public/Blueprint/DMXPixelMappingSubsystem.h
+++ b/Engine/Plugins/VirtualProduction/DMX/DMXPixelMapping/Source/DMXPixelMappingRuntime/Public/Blueprint/DMXPixelMappingSubsystem.h
@@ -2,10 +2,11 @@
 
 #pragma once
 
-#include "DMXPixelMappingRuntimeCommon.h"
 #include "Subsystems/EngineSubsystem.h"
+
 #include "DMXPixelMappingSubsystem.generated.h"
 
+
 UCLASS()
 class DMXPIXELMAPPINGRUNTIME_API UDMXPixelMappingSubsystem
 	: public UEngineSubsystem
diff --git a/Engine/Plugins/VirtualProduction/DMX/DMXPixelMapping/Source/DMXPixelMappingRuntime/Public/ColorSpace/DMXPixelMappingColorSpace.h b/Engine/Plugins/VirtualProduction/DMX/DMXPixelMapping/Source/DMXPixelMappingRuntime/Public/ColorSpace/DMXPixelMappingColorSpace.h
--- a/Engine/Plugins/VirtualProduction/DMX/DMXPixelMapping/Source/DMXPixelMappingRuntime/Public/ColorSpace/DMXPixelMappingColorSpace.h
+++ b/Engine/Plugins/VirtualProduction/DMX/DMXPixelMapping/Source/DMXPixelMappingRuntime/Public/ColorSpace/DMXPixelMappingColorSpace.h
@@ -3,8 +3,6 @@
 #pragma once
 
 #include "DMXAttribute.h"
-
-#include "CoreMinimal.h"
 #include "UObject/Object.h"
 
 #include "DMXPixelMappingColorSpace.generated.h"
diff --git a/Engine/Plugins/VirtualProduction/DMX/DMXPixelMapping/Source/DMXPixelMappingRuntime/Public/Components/DMXPixelMappingBaseComponent.h b/Engine/Plugins/VirtualProduction/DMX/DMXPixelMapping/Source/DMXPixelMappingRuntime/Public/Components/DMXPixelMappingBaseComponent.h
--- a/Engine/Plugins/VirtualProduction/DMX/DMXPixelMapping/Source/DMXPixelMappingRuntime/Public/Components/DMXPixelMappingBaseComponent.h
+++ b/Engine/Plugins/VirtualProduction/DMX/DMXPixelMapping/Source/DMXPixelMappingRuntime/Public/Components/DMXPixelMappingBaseComponent.h
@@ -2,7 +2,6 @@
 
 #pragma once
 
-#include "DMXPixelMappingRuntimeCommon.h"
 #include "UObject/Object.h"
 
 #include "DMXPixelMappingBaseComponent.generated.h"
@@ -12,6 +11,8 @@ DECLARE_STATS_GROUP(TEXT("DMXPixelMapping"), STATGROUP_DMXPIXELMAPPING, STATCAT_
 
 class FDMXPixelMappingComponentTemplate;
 class UDMXPixelMappingRendererComponent;
+class UDMXPixelMapping;
+class UDMXPixelMappingRootComponent;
 
 
 /**
@@ -23,6 +24,11 @@ class DMXPIXELMAPPINGRUNTIME_API UDMXPixelMappingBaseComponent
 {
 	GENERATED_BODY()
 
+	using TComponentPredicate = TFunctionRef<void(UDMXPixelMappingBaseComponent*)>;
+
+	template <typename Type>
+	using TComponentPredicateType = TFunctionRef<void(Type*)>;
+
 	DECLARE_EVENT_TwoParams(UDMXPixelMappingBaseComponent, FDMXPixelMappingOnComponentAdded, UDMXPixelMapping* /** PixelMapping */, UDMXPixelMappingBaseComponent* /** AddedComponent */);
 	DECLARE_EVENT_TwoParams(UDMXPixelMappingBaseComponent, FDMXPixelMappingOnComponentRemoved, UDMXPixelMapping* /** PixelMapping */, UDMXPixelMappingBaseComponent* /** RemovedComponent */);
 	DECLARE_EVENT_OneParam(UDMXPixelMappingBaseComponent, FDMXPixelMappingOnComponentRenamed, UDMXPixelMappingBaseComponent* /** RenamedComponent */);
@@ -109,7 +115,7 @@ public:
 	 *
 	 * @param bIsRecursive		Should it loop recursively
 	 */
-	void ForEachChild(TComponentPredicate Predicate, bool bIsRecursive);
+	void ForEachChild(TFunctionRef<void(UDMXPixelMappingBaseComponent*)> Predicate, bool bIsRecursive);
 
 	/**
 	 * Looking for the first child by given Class
diff --git a/Engine/Plugins/VirtualProduction/DMX/DMXPixelMapping/Source/DMXPixelMappingRuntime/Public/Components/DMXPixelMappingMatrixComponent.h b/Engine/Plugins/VirtualProduction/DMX/DMXPixelMapping/Source/DMXPixelMappingRuntime/Public/Components/DMXPixelMappingMatrixComponent.h
--- a/Engine/Plugins/VirtualProduction/DMX/DMXPixelMapping/Source/DMXPixelMappingRuntime/Public/Components/DMXPixelMappingMatrixComponent.h
+++ b/Engine/Plugins/VirtualProduction/DMX/DMXPixelMapping/Source/DMXPixelMappingRuntime/Public/Components/DMXPixelMappingMatrixComponent.h
@@ -2,17 +2,18 @@
 
 #pragma once
 
-#include "DMXPixelMappingOutputDMXComponent.h"
-
 #include "Components/DMXPixelMappingOutputComponent.h"
+#include "DMXPixelMappingOutputDMXComponent.h"
 #include "Library/DMXEntityReference.h"
 #include "Library/DMXEntityFixtureType.h"
+
 #include "DMXPixelMappingMatrixComponent.generated.h"
 
 enum class EDMXColorMode : uint8;
 class UDMXLibrary;
 class UDMXPixelMappingColorSpace;
 class UDMXPixelMappingLayoutScript;
+class UDMXPixelMappingMatrixCellComponent;
 
 
 /**
diff --git a/Engine/Plugins/VirtualProduction/DMX/DMXPixelMapping/Source/DMXPixelMappingRuntime/Public/DMXPixelMapping.h b/Engine/Plugins/VirtualProduction/DMX/DMXPixelMapping/Source/DMXPixelMappingRuntime/Public/DMXPixelMapping.h
--- a/Engine/Plugins/VirtualProduction/DMX/DMXPixelMapping/Source/DMXPixelMappingRuntime/Public/DMXPixelMapping.h
+++ b/Engine/Plugins/VirtualProduction/DMX/DMXPixelMapping/Source/DMXPixelMappingRuntime/Public/DMXPixelMapping.h
@@ -2,15 +2,13 @@
 
 #pragma once
 
-#include "DMXPixelMappingRuntimeCommon.h"
-
 #include "UObject/Object.h"
 #include "DMXPixelMapping.generated.h"
 
+class SWidget;
 class UDMXEntityFixturePatch;
 class UDMXPixelMappingBaseComponent;
-
-class SWidget;
+class UDMXPixelMappingOutputComponent;
 class UTexture;
 
 #if WITH_EDITOR
@@ -25,6 +23,12 @@ UCLASS(BlueprintType, Blueprintable)
 class DMXPIXELMAPPINGRUNTIME_API UDMXPixelMapping
 	: public UObject
 {
+	using TComponentPredicate = TFunctionRef<void(UDMXPixelMappingBaseComponent*)>;
+
+	template <typename Type>
+	using TComponentPredicateType = TFunctionRef<void(Type*)>;
+
+
 	GENERATED_BODY()
 public:
 	//~ Begin UObject implementation
diff --git a/Engine/Plugins/VirtualProduction/DMX/DMXPixelMapping/Source/DMXPixelMappingRuntime/Public/DMXPixelMappingRuntimeCommon.h b/Engine/Plugins/VirtualProduction/DMX/DMXPixelMapping/Source/DMXPixelMappingRuntime/Public/DMXPixelMappingRuntimeCommon.h
deleted file mode 100644
diff --git a/Engine/Plugins/VirtualProduction/DMX/DMXPixelMapping/Source/DMXPixelMappingRuntime/Public/DMXPixelMappingRuntimeLog.h b/Engine/Plugins/VirtualProduction/DMX/DMXPixelMapping/Source/DMXPixelMappingRuntime/Public/DMXPixelMappingRuntimeLog.h
new file mode 100644
--- /dev/null
+++ b/Engine/Plugins/VirtualProduction/DMX/DMXPixelMapping/Source/DMXPixelMappingRuntime/Public/DMXPixelMappingRuntimeLog.h
@@ -0,0 +1,8 @@
+// Copyright Epic Games, Inc. All Rights Reserved.
+
+#pragma once
+
+#include "Logging/LogMacros.h"
+
+
+DECLARE_LOG_CATEGORY_EXTERN(LogDMXPixelMappingRuntime, All, All);
diff --git a/Engine/Plugins/VirtualProduction/DMX/DMXPixelMapping/Source/DMXPixelMappingRuntime/Public/DMXPixelMappingRuntimeUtils.h b/Engine/Plugins/VirtualProduction/DMX/DMXPixelMapping/Source/DMXPixelMappingRuntime/Public/DMXPixelMappingRuntimeUtils.h
--- a/Engine/Plugins/VirtualProduction/DMX/DMXPixelMapping/Source/DMXPixelMappingRuntime/Public/DMXPixelMappingRuntimeUtils.h
+++ b/Engine/Plugins/VirtualProduction/DMX/DMXPixelMapping/Source/DMXPixelMappingRuntime/Public/DMXPixelMappingRuntimeUtils.h
@@ -2,7 +2,6 @@
 
 #pragma once
 
-#include "CoreMinimal.h"
 #include "Templates/SubclassOf.h"
 
 struct FDMXAttributeName;
@@ -14,6 +13,7 @@ class UDMXPixelMappingBaseComponent;
 struct FPropertyChangedChainEvent;
 #endif
 
+
 class DMXPIXELMAPPINGRUNTIME_API FDMXPixelMappingRuntimeUtils
 {
 public:
diff --git a/Engine/Plugins/VirtualProduction/DMX/DMXPixelMapping/Source/DMXPixelMappingRuntime/Public/LayoutScripts/DMXPixelMappingLayoutScript.h b/Engine/Plugins/VirtualProduction/DMX/DMXPixelMapping/Source/DMXPixelMappingRuntime/Public/LayoutScripts/DMXPixelMappingLayoutScript.h
--- a/Engine/Plugins/VirtualProduction/DMX/DMXPixelMapping/Source/DMXPixelMappingRuntime/Public/LayoutScripts/DMXPixelMappingLayoutScript.h
+++ b/Engine/Plugins/VirtualProduction/DMX/DMXPixelMapping/Source/DMXPixelMappingRuntime/Public/LayoutScripts/DMXPixelMappingLayoutScript.h
@@ -2,7 +2,7 @@
 
 #pragma once
 
-#include "CoreMinimal.h"
+#include "HAL/Platform.h"
 #include "UObject/Object.h"
 
 #include "DMXPixelMappingLayoutScript.generated.h"
@@ -10,7 +10,6 @@
 class UDMXEntityFixturePatch;
 class UDMXLibrary;
 class UDMXPixelMappingOutputComponent;
-
 class UTextureRenderTarget2D;
 
 
diff --git a/Engine/Plugins/VirtualProduction/DMX/DMXPixelMapping/Source/DMXPixelMappingRuntime/Public/Templates/DMXPixelMappingComponentTemplate.h b/Engine/Plugins/VirtualProduction/DMX/DMXPixelMapping/Source/DMXPixelMappingRuntime/Public/Templates/DMXPixelMappingComponentTemplate.h
--- a/Engine/Plugins/VirtualProduction/DMX/DMXPixelMapping/Source/DMXPixelMappingRuntime/Public/Templates/DMXPixelMappingComponentTemplate.h
+++ b/Engine/Plugins/VirtualProduction/DMX/DMXPixelMapping/Source/DMXPixelMappingRuntime/Public/Templates/DMXPixelMappingComponentTemplate.h
@@ -2,8 +2,6 @@
 
 #pragma once
 
-#include "DMXPixelMappingRuntimeCommon.h"
-
 #include "Library/DMXEntityReference.h"
 
 #include "Templates/SharedPointer.h"
