DMX - Fix CIS Build Engine Localization

diff --git a/Engine/Plugins/VirtualProduction/DMX/DMXPixelMapping/Source/DMXPixelMappingEditor/Private/Widgets/SDMXPixelMappingHierarchyRow.cpp b/Engine/Plugins/VirtualProduction/DMX/DMXPixelMapping/Source/DMXPixelMappingEditor/Private/Widgets/SDMXPixelMappingHierarchyRow.cpp
--- a/Engine/Plugins/VirtualProduction/DMX/DMXPixelMapping/Source/DMXPixelMappingEditor/Private/Widgets/SDMXPixelMappingHierarchyRow.cpp
+++ b/Engine/Plugins/VirtualProduction/DMX/DMXPixelMapping/Source/DMXPixelMappingEditor/Private/Widgets/SDMXPixelMappingHierarchyRow.cpp
@@ -172,7 +172,7 @@ bool SDMXPixelMappingHierarchyRow::OnVerifyNameTextChanged(const FText& InText,
 			}) == nullptr;
 		if (!bUniqueName)
 		{
-			OutErrorMessage = LOCTEXT("RenameComponentTransaction", "A component with this name already exists");
+			OutErrorMessage = LOCTEXT("InvalidComponentNameInfo", "A component with this name already exists");
 		}
 
 		return bUniqueName;
