DMX - Fix FDMXMVRFixturePatcherSettings::FixturePatchNameDisplayMode is not properly initialized, remove missing style from pixel mapping - drawing it would cause bad results

diff --git a/Engine/Plugins/VirtualProduction/DMX/DMXEngine/Source/DMXEditor/Public/DMXEditorSettings.h b/Engine/Plugins/VirtualProduction/DMX/DMXEngine/Source/DMXEditor/Public/DMXEditorSettings.h
--- a/Engine/Plugins/VirtualProduction/DMX/DMXEngine/Source/DMXEditor/Public/DMXEditorSettings.h
+++ b/Engine/Plugins/VirtualProduction/DMX/DMXEngine/Source/DMXEditor/Public/DMXEditorSettings.h
@@ -57,7 +57,7 @@ struct FDMXMVRFixturePatcherSettings
 	bool bMonitorEnabled = false;
 
 	UPROPERTY()
-	EDMXFixturePatcherNameDisplayMode FixturePatchNameDisplayMode;
+	EDMXFixturePatcherNameDisplayMode FixturePatchNameDisplayMode = EDMXFixturePatcherNameDisplayMode::FixturePatchName;
 };
 
 
diff --git a/Engine/Plugins/VirtualProduction/DMX/DMXPixelMapping/Source/DMXPixelMappingEditor/Private/DMXPixelMappingEditorStyle.cpp b/Engine/Plugins/VirtualProduction/DMX/DMXPixelMapping/Source/DMXPixelMappingEditor/Private/DMXPixelMappingEditorStyle.cpp
--- a/Engine/Plugins/VirtualProduction/DMX/DMXPixelMapping/Source/DMXPixelMappingEditor/Private/DMXPixelMappingEditorStyle.cpp
+++ b/Engine/Plugins/VirtualProduction/DMX/DMXPixelMapping/Source/DMXPixelMappingEditor/Private/DMXPixelMappingEditorStyle.cpp
@@ -4,12 +4,13 @@
 
 #include "Brushes/SlateImageBrush.h"
 #include "Framework/Application/SlateApplication.h"
-#include "Misc/Paths.h"
-#include "Styling/SlateStyleRegistry.h"
 #include "Interfaces/IPluginManager.h"
+#include "Misc/Paths.h"
 #include "Rendering/SlateRenderer.h"
+#include "Styling/SlateStyleRegistry.h"
 #include "Styling/SlateStyle.h"
 
+
 TSharedPtr< FSlateStyleSet > FDMXPixelMappingEditorStyle::DMXPixelMappingEditorStyleInstance = nullptr;
 
 void FDMXPixelMappingEditorStyle::Initialize()
@@ -82,8 +83,6 @@ TSharedRef<FSlateStyleSet> FDMXPixelMappingEditorStyle::Create()
 	Style->Set("DMXPixelMappingEditor.PlayDMX.Small", new IMAGE_BRUSH("icon_DMXPixelMappingEditor_PlayDMX_40x", Icon20x20));
 	Style->Set("DMXPixelMappingEditor.StopPlayingDMX", new IMAGE_BRUSH("icon_DMXPixelMappingEditor_StopPlayingDMX_40x", Icon40x40));
 	Style->Set("DMXPixelMappingEditor.StopPlayingDMX.Small", new IMAGE_BRUSH("icon_DMXPixelMappingEditor_StopPlayingDMX_40x", Icon20x20));
-	
-	Style->Set("DMXPixelMappingEditor.ComponentBorder", new BORDER_BRUSH("Common/Selection", FMargin(2.f), FLinearColor::White));
 
 	return Style;
 }
diff --git a/Engine/Plugins/VirtualProduction/DMX/DMXPixelMapping/Source/DMXPixelMappingEditor/Private/Widgets/SDMXPixelMappingOutputComponent.cpp b/Engine/Plugins/VirtualProduction/DMX/DMXPixelMapping/Source/DMXPixelMappingEditor/Private/Widgets/SDMXPixelMappingOutputComponent.cpp
--- a/Engine/Plugins/VirtualProduction/DMX/DMXPixelMapping/Source/DMXPixelMappingEditor/Private/Widgets/SDMXPixelMappingOutputComponent.cpp
+++ b/Engine/Plugins/VirtualProduction/DMX/DMXPixelMapping/Source/DMXPixelMappingEditor/Private/Widgets/SDMXPixelMappingOutputComponent.cpp
@@ -125,7 +125,6 @@ TSharedRef<SWidget> SDMXPixelMappingOutputComponent::CreateContent()
 {
 	const TSharedRef<SBorder> Content =
 		SNew(SBorder)
-		.BorderImage(FDMXPixelMappingEditorStyle::Get().GetBrush("DMXPixelMappingEditor.ComponentBorder"))
 		.BorderBackgroundColor_Lambda([this]()
 			{
 				return Model->GetColor();
