DMX - Fix fixture patch list in pixel mapping never shows a scrollbar

diff --git a/Engine/Plugins/VirtualProduction/DMX/DMXEngine/Source/DMXEditor/Private/Widgets/FixturePatch/SDMXReadOnlyFixturePatchList.cpp b/Engine/Plugins/VirtualProduction/DMX/DMXEngine/Source/DMXEditor/Private/Widgets/FixturePatch/SDMXReadOnlyFixturePatchList.cpp
--- a/Engine/Plugins/VirtualProduction/DMX/DMXEngine/Source/DMXEditor/Private/Widgets/FixturePatch/SDMXReadOnlyFixturePatchList.cpp
+++ b/Engine/Plugins/VirtualProduction/DMX/DMXEngine/Source/DMXEditor/Private/Widgets/FixturePatch/SDMXReadOnlyFixturePatchList.cpp
@@ -115,7 +115,6 @@ void SDMXReadOnlyFixturePatchList::Construct(const FArguments& InArgs)
 
 			// ListView section
 			+ SVerticalBox::Slot()
-			.AutoHeight()
 			[
 				SAssignNew(ListView, SListView<TSharedPtr<FDMXEntityFixturePatchRef>>)
 				.HeaderRow(GenerateHeaderRow())
diff --git a/Engine/Plugins/VirtualProduction/DMX/DMXPixelMapping/Source/DMXPixelMappingEditor/Private/Views/SDMXPixelMappingDMXLibraryView.cpp b/Engine/Plugins/VirtualProduction/DMX/DMXPixelMapping/Source/DMXPixelMappingEditor/Private/Views/SDMXPixelMappingDMXLibraryView.cpp
--- a/Engine/Plugins/VirtualProduction/DMX/DMXPixelMapping/Source/DMXPixelMappingEditor/Private/Views/SDMXPixelMappingDMXLibraryView.cpp
+++ b/Engine/Plugins/VirtualProduction/DMX/DMXPixelMapping/Source/DMXPixelMappingEditor/Private/Views/SDMXPixelMappingDMXLibraryView.cpp
@@ -22,6 +22,7 @@
 #include "ViewModels/DMXPixelMappingDMXLibraryViewModel.h"
 #include "Widgets/Images/SImage.h"
 #include "Widgets/Input/SButton.h"
+#include "Widgets/Layout/SWidgetSwitcher.h"
 #include "Widgets/SBoxPanel.h"
 #include "Widgets/SDMXPixelMappingFixturePatchList.h"
 
@@ -172,27 +173,31 @@ void SDMXPixelMappingDMXLibraryView::Construct(const FArguments& InArgs, const T
 				]
 			]
 		
-			// Fixture patch list
 			+ SVerticalBox::Slot()
-			.AutoHeight()
+			.FillHeight(1.f)
 			[
 				SNew(SBorder)
 				.BorderImage(FAppStyle::Get().GetBrush("Brushes.Recessed"))
 				.Padding(FMargin(8.f, 2.f, 2.f, 0.f))
 				[
-					SAssignNew(FixturePatchList, SDMXPixelMappingFixturePatchList, Toolkit, ViewModel)
-				]
-			]
+					SAssignNew(ListOrAllPatchesAddedSwitcher, SWidgetSwitcher)
 
-			// 'All patches added to pixelmapping' info
-			+ SVerticalBox::Slot()
-			.FillHeight(1.f)
-			.HAlign(HAlign_Center)
-			.VAlign(VAlign_Center)
-			[
-				SAssignNew(AllPatchesAddedTextBlock, STextBlock)
-				.Text(LOCTEXT("AllPatchesAddedHint", "All Patches added to Pixel Mapping"))
-				.Font(FAppStyle::GetFontStyle("PropertyWindow.NormalFont"))
+					// Fixture patch list
+					+ SWidgetSwitcher::Slot()
+					[
+						SAssignNew(FixturePatchList, SDMXPixelMappingFixturePatchList, Toolkit, ViewModel)
+					]
+
+					// All patches added info box
+					+ SWidgetSwitcher::Slot()
+					.HAlign(HAlign_Center)
+					.VAlign(VAlign_Center)
+					[
+						SAssignNew(AllPatchesAddedTextBlock, STextBlock)
+						.Text(LOCTEXT("AllPatchesAddedHint", "All Patches added to Pixel Mapping"))
+						.Font(FAppStyle::GetFontStyle("PropertyWindow.NormalFont"))
+					]
+				]
 			]
 		]
 	];
@@ -288,14 +293,18 @@ void SDMXPixelMappingDMXLibraryView::ForceRefresh()
 			});
 		FixturePatchList->SetExcludedFixturePatches(HiddenFixturePatches);
 
-		// Only show options to add or select when fixture patches are available
+		// Only show the list and options to add or select when fixture patches are available
 		const EVisibility AddFixturePatchOptionsVisibility = FixturePatchesInDMXLibrary.Num() > HiddenFixturePatches.Num() ? EVisibility::Visible : EVisibility::Collapsed;
 		AddPatchesHorizontalBox->SetVisibility(AddFixturePatchOptionsVisibility);
-		FixturePatchList->SetVisibility(AddFixturePatchOptionsVisibility);
 
-		// Show an info if all patches are added to pixel mapping
-		const EVisibility AllPatchesAddedVisibility = !FixturePatchesInDMXLibrary.IsEmpty() && FixturePatchesInDMXLibrary.Num() == HiddenFixturePatches.Num() ? EVisibility::Visible : EVisibility::Collapsed;
-		AllPatchesAddedTextBlock->SetVisibility(AllPatchesAddedVisibility);
+		if (AddFixturePatchOptionsVisibility == EVisibility::Visible)
+		{
+			ListOrAllPatchesAddedSwitcher->SetActiveWidget(FixturePatchList.ToSharedRef());
+		}
+		else
+		{
+			ListOrAllPatchesAddedSwitcher->SetActiveWidget(AllPatchesAddedTextBlock.ToSharedRef());
+		}
 	}
 }
 
diff --git a/Engine/Plugins/VirtualProduction/DMX/DMXPixelMapping/Source/DMXPixelMappingEditor/Private/Views/SDMXPixelMappingDMXLibraryView.h b/Engine/Plugins/VirtualProduction/DMX/DMXPixelMapping/Source/DMXPixelMappingEditor/Private/Views/SDMXPixelMappingDMXLibraryView.h
--- a/Engine/Plugins/VirtualProduction/DMX/DMXPixelMapping/Source/DMXPixelMappingEditor/Private/Views/SDMXPixelMappingDMXLibraryView.h
+++ b/Engine/Plugins/VirtualProduction/DMX/DMXPixelMapping/Source/DMXPixelMappingEditor/Private/Views/SDMXPixelMappingDMXLibraryView.h
@@ -15,6 +15,7 @@ class SBorder;
 class SDMXPixelMappingFixturePatchList;
 class SHorizontalBox;
 class STextBlock;
+class SWidgetSwitcher;
 class UDMXEntity;
 class UDMXLibrary;
 class UDMXPixelMapping;
@@ -96,12 +97,15 @@ private:
 	/** Text block displaying the selected fixture group name */
 	TSharedPtr<STextBlock> FixtureGroupNameTextBlock;
 
-	/** Text block displaying 'all patches added' when no patches are available to add from the dmx library */
-	TSharedPtr<STextBlock> AllPatchesAddedTextBlock;
+	/** Switches between the list and the all patches added text blocok */
+	TSharedPtr<SWidgetSwitcher> ListOrAllPatchesAddedSwitcher;
 
 	/** List of fixture patches the user can select from */
 	TSharedPtr<SDMXPixelMappingFixturePatchList> FixturePatchList;
 
+	/** Text block displaying 'all patches added' when no patches are available to add from the dmx library */
+	TSharedPtr<STextBlock> AllPatchesAddedTextBlock;
+
 	/** Horizontal box containing the 'add selected patches' and 'add all patches' buttons */
 	TSharedPtr<SHorizontalBox> AddPatchesHorizontalBox;
 
