DMX - Fix pixel mapping components don't show the intent border style in editor, components and selection are now clearly visible again

diff --git a/Engine/Plugins/VirtualProduction/DMX/DMXPixelMapping/Source/DMXPixelMappingEditor/Private/DMXPixelMappingEditorStyle.cpp b/Engine/Plugins/VirtualProduction/DMX/DMXPixelMapping/Source/DMXPixelMappingEditor/Private/DMXPixelMappingEditorStyle.cpp
--- a/Engine/Plugins/VirtualProduction/DMX/DMXPixelMapping/Source/DMXPixelMappingEditor/Private/DMXPixelMappingEditorStyle.cpp
+++ b/Engine/Plugins/VirtualProduction/DMX/DMXPixelMapping/Source/DMXPixelMappingEditor/Private/DMXPixelMappingEditorStyle.cpp
@@ -7,6 +7,7 @@
 #include "Interfaces/IPluginManager.h"
 #include "Misc/Paths.h"
 #include "Rendering/SlateRenderer.h"
+#include "Styling/SlateBrush.h"
 #include "Styling/SlateStyleRegistry.h"
 #include "Styling/SlateStyle.h"
 
@@ -31,13 +32,13 @@ void FDMXPixelMappingEditorStyle::Shutdown()
 
 FName FDMXPixelMappingEditorStyle::GetStyleSetName()
 {
-	static FName StyleSetName(TEXT("DMXPixelMappingEditorStyle"));
+	const FName StyleSetName(TEXT("DMXPixelMappingEditorStyle"));
 	return StyleSetName;
 }
 
 FString RelativePathToPluginPath(const FString& RelativePath, const ANSICHAR* Extension)
 {
-	static FString ContentDir = IPluginManager::Get().FindPlugin(TEXT("DMXPixelMapping"))->GetContentDir();
+	const FString ContentDir = IPluginManager::Get().FindPlugin(TEXT("DMXPixelMapping"))->GetContentDir();
 	return (ContentDir / RelativePath) + Extension;
 }
 
@@ -50,16 +51,17 @@ FString RelativePathToPluginPath(const FString& RelativePath, const ANSICHAR* Ex
 #define BORDER_BRUSH( RelativePath, ... ) FSlateBorderBrush( Style->RootToContentDir( RelativePath, TEXT(".png") ), __VA_ARGS__ )
 #define DEFAULT_FONT(...) FCoreStyle::GetDefaultFontStyle(__VA_ARGS__)
 
-const FVector2D Icon8x8(8.0f, 8.0f);
-const FVector2D Icon12x12(12.0f, 12.0f);
-const FVector2D Icon16x16(16.0f, 16.0f);
-const FVector2D Icon20x20(20.0f, 20.0f);
-const FVector2D Icon40x40(40.0f, 40.0f);
-const FVector2D Icon64x64(64.0f, 64.0f);
-
 TSharedRef<FSlateStyleSet> FDMXPixelMappingEditorStyle::Create()
 {
-	TSharedRef<FSlateStyleSet> Style = MakeShareable(new FSlateStyleSet("DMXPixelMappingEditorStyle"));
+	const FVector2D Icon8x8(8.0f, 8.0f);
+	const FVector2D Icon12x12(12.0f, 12.0f);
+	const FVector2D Icon16x16(16.0f, 16.0f);
+	const FVector2D Icon20x20(20.0f, 20.0f);
+	const FVector2D Icon40x40(40.0f, 40.0f);
+	const FVector2D Icon64x64(64.0f, 64.0f);
+
+
+	const TSharedRef<FSlateStyleSet> Style = MakeShareable(new FSlateStyleSet("DMXPixelMappingEditorStyle"));
 
 	TSharedPtr<IPlugin> Plugin = IPluginManager::Get().FindPlugin(TEXT("DMXPixelMapping"));
 	if (Plugin.IsValid())
@@ -84,6 +86,13 @@ TSharedRef<FSlateStyleSet> FDMXPixelMappingEditorStyle::Create()
 	Style->Set("DMXPixelMappingEditor.StopPlayingDMX", new IMAGE_BRUSH("icon_DMXPixelMappingEditor_StopPlayingDMX_40x", Icon40x40));
 	Style->Set("DMXPixelMappingEditor.StopPlayingDMX.Small", new IMAGE_BRUSH("icon_DMXPixelMappingEditor_StopPlayingDMX_40x", Icon20x20));
 
+	// Component border style
+	FSlateBrush* ComponentBorderBrush = new FSlateBrush();
+	ComponentBorderBrush->Margin = FMargin(1.f);
+	ComponentBorderBrush->DrawAs = ESlateBrushDrawType::Border;
+	ComponentBorderBrush->TintColor = FLinearColor::White;
+	Style->Set("DMXPixelMappingEditor.ComponentBorder", ComponentBorderBrush);
+
 	return Style;
 }
 
diff --git a/Engine/Plugins/VirtualProduction/DMX/DMXPixelMapping/Source/DMXPixelMappingEditor/Private/ViewModels/DMXPixelMappingOutputComponentModel.cpp b/Engine/Plugins/VirtualProduction/DMX/DMXPixelMapping/Source/DMXPixelMappingEditor/Private/ViewModels/DMXPixelMappingOutputComponentModel.cpp
--- a/Engine/Plugins/VirtualProduction/DMX/DMXPixelMapping/Source/DMXPixelMappingEditor/Private/ViewModels/DMXPixelMappingOutputComponentModel.cpp
+++ b/Engine/Plugins/VirtualProduction/DMX/DMXPixelMapping/Source/DMXPixelMappingEditor/Private/ViewModels/DMXPixelMappingOutputComponentModel.cpp
@@ -166,7 +166,8 @@ FLinearColor FDMXPixelMappingOutputComponentModel::GetColor() const
 			}
 			else
 			{
-				return OutputComponent->GetEditorColor().CopyWithNewOpacity(0.75f);
+				const float Opacity = OutputComponent->GetEditorColor().A * .75f;
+				return OutputComponent->GetEditorColor().CopyWithNewOpacity(Opacity);
 			}
 		}
 	}
diff --git a/Engine/Plugins/VirtualProduction/DMX/DMXPixelMapping/Source/DMXPixelMappingEditor/Private/Widgets/SDMXPixelMappingOutputComponent.cpp b/Engine/Plugins/VirtualProduction/DMX/DMXPixelMapping/Source/DMXPixelMappingEditor/Private/Widgets/SDMXPixelMappingOutputComponent.cpp
--- a/Engine/Plugins/VirtualProduction/DMX/DMXPixelMapping/Source/DMXPixelMappingEditor/Private/Widgets/SDMXPixelMappingOutputComponent.cpp
+++ b/Engine/Plugins/VirtualProduction/DMX/DMXPixelMapping/Source/DMXPixelMappingEditor/Private/Widgets/SDMXPixelMappingOutputComponent.cpp
@@ -55,6 +55,7 @@ void IDMXPixelMappingOutputComponentWidgetInterface::RemoveFromCanvas()
 	ParentCanvas.Reset();
 }
 
+
 void SDMXPixelMappingOutputComponent::Construct(const FArguments& InArgs, const TSharedRef<FDMXPixelMappingToolkit>& InToolkit, TWeakObjectPtr<UDMXPixelMappingOutputComponent> OutputComponent)
 {
 	if (!OutputComponent.IsValid())
@@ -65,9 +66,6 @@ void SDMXPixelMappingOutputComponent::Construct(const FArguments& InArgs, const
 	WeakToolkit = InToolkit;
 	Model = MakeShared<FDMXPixelMappingOutputComponentModel>(InToolkit, OutputComponent);
 
-	BorderBrush.DrawAs = ESlateBrushDrawType::Border;
-	BorderBrush.Margin = FMargin(1.f);
-
 	ChildSlot
 	[
 		SAssignNew(ComponentBox, SBox)
@@ -125,6 +123,7 @@ TSharedRef<SWidget> SDMXPixelMappingOutputComponent::CreateContent()
 {
 	const TSharedRef<SBorder> Content =
 		SNew(SBorder)
+		.BorderImage(FDMXPixelMappingEditorStyle::Get().GetBrush("DMXPixelMappingEditor.ComponentBorder"))
 		.BorderBackgroundColor_Lambda([this]()
 			{
 				return Model->GetColor();
@@ -350,14 +349,10 @@ public:
 		Universe = InArgs._Universe;
 		StartingChannel = InArgs._StartingChannel;
 
-		BorderBrush.DrawAs = ESlateBrushDrawType::Border;
-		BorderBrush.Margin = FMargin(1.f);
-		BorderBrush.TintColor = FLinearColor::White.CopyWithNewOpacity(.4f);
-
 		ChildSlot
 			[
 				SAssignNew(ContentBorder, SBorder)
-				.BorderImage(&BorderBrush)
+				.BorderImage(FDMXPixelMappingEditorStyle::Get().GetBrush("DMXPixelMappingEditor.ComponentBorder"))
 				.HAlign(HAlign_Fill)
 				.VAlign(VAlign_Fill)
 			];
@@ -435,9 +430,6 @@ private:
 	/** Border to hold contents of the widget */
 	TSharedPtr<SBorder> ContentBorder;
 
-	/** The brush used for the border */
-	FSlateBrush BorderBrush;
-
 	/** The model used by this widget */
 	TSharedPtr<FDMXPixelMappingScreenComponentModel> Model;
 
@@ -624,9 +616,6 @@ void SDMXPixelMappingScreenComponent::Construct(const FArguments& InArgs, const
 	WeakToolkit = InToolkit;
 	Model = MakeShared<FDMXPixelMappingScreenComponentModel>(InToolkit, ScreenComponent);
 
-	BorderBrush.DrawAs = ESlateBrushDrawType::Border;
-	BorderBrush.Margin = FMargin(1.f);
-
 	ChildSlot
 	[
 		SNew(SBox)
@@ -651,10 +640,10 @@ void SDMXPixelMappingScreenComponent::Construct(const FArguments& InArgs, const
 			+ SOverlay::Slot() // Required to correctly draw the color around the grid
 			[
 				SNew(SBorder)
-				.BorderImage_Lambda([this]()
+				.BorderImage(FDMXPixelMappingEditorStyle::Get().GetBrush("DMXPixelMappingEditor.ComponentBorder"))
+				.BorderBackgroundColor_Lambda([this]()
 					{
-						BorderBrush.TintColor = Model->GetColor();
-						return &BorderBrush;
+						return Model->GetColor();
 					})
 			]
 		]
diff --git a/Engine/Plugins/VirtualProduction/DMX/DMXPixelMapping/Source/DMXPixelMappingEditor/Private/Widgets/SDMXPixelMappingOutputComponent.h b/Engine/Plugins/VirtualProduction/DMX/DMXPixelMapping/Source/DMXPixelMappingEditor/Private/Widgets/SDMXPixelMappingOutputComponent.h
--- a/Engine/Plugins/VirtualProduction/DMX/DMXPixelMapping/Source/DMXPixelMappingEditor/Private/Widgets/SDMXPixelMappingOutputComponent.h
+++ b/Engine/Plugins/VirtualProduction/DMX/DMXPixelMapping/Source/DMXPixelMappingEditor/Private/Widgets/SDMXPixelMappingOutputComponent.h
@@ -3,7 +3,6 @@
 #pragma once
 
 #include "Engine/EngineTypes.h"
-#include "Styling/SlateBrush.h"
 #include "Widgets/SCompoundWidget.h"
 #include "Widgets/Layout/SConstraintCanvas.h"
 
@@ -93,9 +92,6 @@ private:
 	/** Creates the child slot that displays info about the patch such as Addresses or the Fixture ID */
 	void CreatePatchInfoChildSlot();
 
-	/** Brush for the outermost Bborder */
-	FSlateBrush BorderBrush;
-
 	/** The box that is shown */
 	TSharedPtr<SBox> ComponentBox;
 
@@ -161,9 +157,6 @@ private:
 	/** The canvas slot of the component widget */
 	SConstraintCanvas::FSlot* Slot = nullptr;
 
-	/** Brush for the outermost Bborder */
-	FSlateBrush BorderBrush;
-
 	/** Border that holds the content */
 	TSharedPtr<SBorder> ContentBorder;
 
diff --git a/Engine/Plugins/VirtualProduction/DMX/DMXPixelMapping/Source/DMXPixelMappingRuntime/Private/Components/DMXPixelMappingMatrixComponent.cpp b/Engine/Plugins/VirtualProduction/DMX/DMXPixelMapping/Source/DMXPixelMappingRuntime/Private/Components/DMXPixelMappingMatrixComponent.cpp
--- a/Engine/Plugins/VirtualProduction/DMX/DMXPixelMapping/Source/DMXPixelMappingRuntime/Private/Components/DMXPixelMappingMatrixComponent.cpp
+++ b/Engine/Plugins/VirtualProduction/DMX/DMXPixelMapping/Source/DMXPixelMappingRuntime/Private/Components/DMXPixelMappingMatrixComponent.cpp
@@ -408,7 +408,7 @@ FString UDMXPixelMappingMatrixComponent::GetUserName() const
 	const UDMXEntityFixturePatch* Patch = FixturePatchRef.GetFixturePatch();
 	if (Patch && UserName.IsEmpty())
 	{
-		return FString::Printf(TEXT("Fixture Matrix: %s"), *Patch->GetDisplayName());
+		return FString::Printf(TEXT("%s"), *Patch->GetDisplayName());
 	}
 	else
 	{
