DMX - Use star ship play/stop icons in pixel mapping, refactor style code to register and setup like most style classes in engine

diff --git a/Engine/Plugins/VirtualProduction/DMX/DMXPixelMapping/Source/DMXPixelMappingEditor/Private/DMXPixelMappingEditorModule.cpp b/Engine/Plugins/VirtualProduction/DMX/DMXPixelMapping/Source/DMXPixelMappingEditor/Private/DMXPixelMappingEditorModule.cpp
--- a/Engine/Plugins/VirtualProduction/DMX/DMXPixelMapping/Source/DMXPixelMappingEditor/Private/DMXPixelMappingEditorModule.cpp
+++ b/Engine/Plugins/VirtualProduction/DMX/DMXPixelMapping/Source/DMXPixelMappingEditor/Private/DMXPixelMappingEditorModule.cpp
@@ -24,7 +24,6 @@ void FDMXPixelMappingEditorModule::StartupModule()
 	DMXPixelMappingCategory = AssetTools.RegisterAdvancedAssetCategory(FName(TEXT("DMX")), LOCTEXT("DmxCategory", "DMX"));
 	RegisterAssetTypeAction(AssetTools, MakeShared<FAssetTypeActions_DMXPixelMapping>());
 
-	FDMXPixelMappingEditorStyle::Initialize();
 	FDMXPixelMappingEditorCommands::Register();
 
 	// Any attempt to use GEditor right now will fail as it hasn't been initialized yet. Waiting for post engine init resolves that.
@@ -42,8 +41,6 @@ void FDMXPixelMappingEditorModule::ShutdownModule()
 		}
 	}
 
-	FDMXPixelMappingEditorStyle::Shutdown();
-
 	if (UObjectInitialized() && GIsEditor)
 	{
 		UThumbnailManager::Get().UnregisterCustomRenderer(UDMXPixelMapping::StaticClass());
diff --git a/Engine/Plugins/VirtualProduction/DMX/DMXPixelMapping/Source/DMXPixelMappingEditor/Private/DMXPixelMappingEditorStyle.cpp b/Engine/Plugins/VirtualProduction/DMX/DMXPixelMapping/Source/DMXPixelMappingEditor/Private/DMXPixelMappingEditorStyle.cpp
--- a/Engine/Plugins/VirtualProduction/DMX/DMXPixelMapping/Source/DMXPixelMappingEditor/Private/DMXPixelMappingEditorStyle.cpp
+++ b/Engine/Plugins/VirtualProduction/DMX/DMXPixelMapping/Source/DMXPixelMappingEditor/Private/DMXPixelMappingEditorStyle.cpp
@@ -8,51 +8,22 @@
 #include "Misc/Paths.h"
 #include "Rendering/SlateRenderer.h"
 #include "Styling/SlateBrush.h"
+#include "Styling/StyleColors.h"
+#include "Styling/SlateStyleMacros.h"
 #include "Styling/SlateStyleRegistry.h"
 #include "Styling/SlateStyle.h"
 
 
-TSharedPtr< FSlateStyleSet > FDMXPixelMappingEditorStyle::DMXPixelMappingEditorStyleInstance = nullptr;
-
-void FDMXPixelMappingEditorStyle::Initialize()
+FDMXPixelMappingEditorStyle::FDMXPixelMappingEditorStyle()
+	: FSlateStyleSet("DMXPixelMappingEditorStyle")
 {
-	if (!DMXPixelMappingEditorStyleInstance.IsValid())
+	TSharedPtr<IPlugin> Plugin = IPluginManager::Get().FindPlugin(TEXT("DMXPixelMapping"));
+	if (Plugin.IsValid())
 	{
-		DMXPixelMappingEditorStyleInstance = Create();
-		FSlateStyleRegistry::RegisterSlateStyle(*DMXPixelMappingEditorStyleInstance);
+		SetContentRoot(FPaths::Combine(Plugin->GetBaseDir(), TEXT("Content/Slate")));
 	}
-}
+	SetCoreContentRoot(FPaths::EngineContentDir() / TEXT("Slate"));
 
-void FDMXPixelMappingEditorStyle::Shutdown()
-{
-	FSlateStyleRegistry::UnRegisterSlateStyle(*DMXPixelMappingEditorStyleInstance);
-	ensure(DMXPixelMappingEditorStyleInstance.IsUnique());
-	DMXPixelMappingEditorStyleInstance.Reset();
-}
-
-FName FDMXPixelMappingEditorStyle::GetStyleSetName()
-{
-	const FName StyleSetName(TEXT("DMXPixelMappingEditorStyle"));
-	return StyleSetName;
-}
-
-FString RelativePathToPluginPath(const FString& RelativePath, const ANSICHAR* Extension)
-{
-	const FString ContentDir = IPluginManager::Get().FindPlugin(TEXT("DMXPixelMapping"))->GetContentDir();
-	return (ContentDir / RelativePath) + Extension;
-}
-
-#define IMAGE_BRUSH( RelativePath, ... ) FSlateImageBrush( Style->RootToContentDir( RelativePath, TEXT(".png") ), __VA_ARGS__ )
-#define IMAGE_BRUSH_SVG( RelativePath, ... ) FSlateVectorImageBrush( Style->RootToContentDir(RelativePath, TEXT(".svg")), __VA_ARGS__ )
-#define IMAGE_CORE_BRUSH( RelativePath, ... ) FSlateImageBrush( FPaths::EngineContentDir() / "Editor/Slate" / RelativePath + TEXT(".png"), __VA_ARGS__ )
-#define IMAGE_PLUGIN_BRUSH( RelativePath, ... ) FSlateImageBrush( RelativePathToPluginPath( RelativePath, ".png" ), __VA_ARGS__ )
-#define BOX_BRUSH( RelativePath, ... ) FSlateBoxBrush( Style->RootToContentDir( RelativePath, TEXT(".png") ), __VA_ARGS__ )
-#define BOX_CORE_BRUSH( RelativePath, ... ) FSlateBoxBrush( FPaths::EngineContentDir() / "Editor/Slate" / RelativePath + TEXT(".png"), __VA_ARGS__ )
-#define BORDER_BRUSH( RelativePath, ... ) FSlateBorderBrush( Style->RootToContentDir( RelativePath, TEXT(".png") ), __VA_ARGS__ )
-#define DEFAULT_FONT(...) FCoreStyle::GetDefaultFontStyle(__VA_ARGS__)
-
-TSharedRef<FSlateStyleSet> FDMXPixelMappingEditorStyle::Create()
-{
 	const FVector2D Icon8x8(8.0f, 8.0f);
 	const FVector2D Icon12x12(12.0f, 12.0f);
 	const FVector2D Icon16x16(16.0f, 16.0f);
@@ -60,57 +31,40 @@ TSharedRef<FSlateStyleSet> FDMXPixelMappingEditorStyle::Create()
 	const FVector2D Icon40x40(40.0f, 40.0f);
 	const FVector2D Icon64x64(64.0f, 64.0f);
 
-
-	const TSharedRef<FSlateStyleSet> Style = MakeShareable(new FSlateStyleSet("DMXPixelMappingEditorStyle"));
-
-	TSharedPtr<IPlugin> Plugin = IPluginManager::Get().FindPlugin(TEXT("DMXPixelMapping"));
-	if (Plugin.IsValid())
+	// Class styles
 	{
-		Style->SetContentRoot(FPaths::Combine(Plugin->GetBaseDir(), TEXT("Content/Slate")));
+		Set("ClassIcon.DMXPixelMapping", new IMAGE_BRUSH_SVG("DMXPixelMapping_16", Icon16x16));
+		Set("ClassThumbnail.DMXPixelMapping", new IMAGE_BRUSH_SVG("DMXPixelMapping_64", Icon64x64));
 	}
 
-	// Class styles
-	Style->Set("ClassIcon.DMXPixelMapping", new IMAGE_BRUSH_SVG("DMXPixelMapping_16", Icon16x16));
-	Style->Set("ClassThumbnail.DMXPixelMapping", new IMAGE_BRUSH_SVG("DMXPixelMapping_64", Icon64x64));
-
 	// Icons
-	Style->Set("Icons.Preview", new IMAGE_BRUSH_SVG("Preview", Icon16x16));
-	Style->Set("Icons.ZoomToFit", new IMAGE_BRUSH_SVG("ZoomToFit", Icon16x16));
-
-	// Toolbar Icons
-	Style->Set("DMXPixelMappingEditor.AddMapping", new IMAGE_BRUSH("icon_DMXPixelMappingEditor_AddMapping_40x", Icon40x40));
-	Style->Set("DMXPixelMappingEditor.AddMapping.Small", new IMAGE_BRUSH("icon_DMXPixelMappingEditor_AddMapping_40x", Icon20x20));
-
-	Style->Set("DMXPixelMappingEditor.PlayDMX", new IMAGE_BRUSH("icon_DMXPixelMappingEditor_PlayDMX_40x", Icon40x40));
-	Style->Set("DMXPixelMappingEditor.PlayDMX.Small", new IMAGE_BRUSH("icon_DMXPixelMappingEditor_PlayDMX_40x", Icon20x20));
-	Style->Set("DMXPixelMappingEditor.StopPlayingDMX", new IMAGE_BRUSH("icon_DMXPixelMappingEditor_StopPlayingDMX_40x", Icon40x40));
-	Style->Set("DMXPixelMappingEditor.StopPlayingDMX.Small", new IMAGE_BRUSH("icon_DMXPixelMappingEditor_StopPlayingDMX_40x", Icon20x20));
+	{
+		Set("Icons.Preview", new IMAGE_BRUSH_SVG("Preview", Icon16x16));
+		Set("Icons.ZoomToFit", new IMAGE_BRUSH_SVG("ZoomToFit", Icon16x16));
+		Set("Icons.AddMapping", new IMAGE_BRUSH("icon_DMXPixelMappingEditor_AddMapping_40x", Icon16x16));
+		Set("Icons.PlayDMX", new CORE_IMAGE_BRUSH_SVG("Starship/Common/play", Icon16x16, FStyleColors::AccentGreen));
+		Set("Icons.StopPlayingDMX", new CORE_IMAGE_BRUSH_SVG("Starship/Common/stop", Icon16x16, FStyleColors::AccentRed));
+	}
 
 	// Component border style
-	FSlateBrush* ComponentBorderBrush = new FSlateBrush();
-	ComponentBorderBrush->Margin = FMargin(1.f);
-	ComponentBorderBrush->DrawAs = ESlateBrushDrawType::Border;
-	ComponentBorderBrush->TintColor = FLinearColor::White;
-	Style->Set("DMXPixelMappingEditor.ComponentBorder", ComponentBorderBrush);
+	{
+		FSlateBrush* ComponentBorderBrush = new FSlateBrush();
+		ComponentBorderBrush->Margin = FMargin(1.f);
+		ComponentBorderBrush->DrawAs = ESlateBrushDrawType::Border;
+		ComponentBorderBrush->TintColor = FLinearColor::White;
+		Set("DMXPixelMappingEditor.ComponentBorder", ComponentBorderBrush);
+	}
 
-	return Style;
+	FSlateStyleRegistry::RegisterSlateStyle(*this);
 }
 
-#undef IMAGE_BRUSH
-#undef IMAGE_BRUSH_SVG
-#undef IMAGE_CORE_BRUSH
-#undef IMAGE_PLUGIN_BRUSH
-#undef BOX_BRUSH
-#undef BORDER_BRUSH
-#undef BOX_CORE_BRUSH
-#undef DEFAULT_FONT
-
-void FDMXPixelMappingEditorStyle::ReloadTextures()
+FDMXPixelMappingEditorStyle::~FDMXPixelMappingEditorStyle()
 {
-	FSlateApplication::Get().GetRenderer()->ReloadTextureResources();
+	FSlateStyleRegistry::UnRegisterSlateStyle(*this);
 }
 
 const ISlateStyle& FDMXPixelMappingEditorStyle::Get()
 {
-	return *DMXPixelMappingEditorStyleInstance;
+	static const FDMXPixelMappingEditorStyle Inst;
+	return Inst;
 }
diff --git a/Engine/Plugins/VirtualProduction/DMX/DMXPixelMapping/Source/DMXPixelMappingEditor/Private/DMXPixelMappingEditorStyle.h b/Engine/Plugins/VirtualProduction/DMX/DMXPixelMapping/Source/DMXPixelMappingEditor/Private/DMXPixelMappingEditorStyle.h
--- a/Engine/Plugins/VirtualProduction/DMX/DMXPixelMapping/Source/DMXPixelMappingEditor/Private/DMXPixelMappingEditorStyle.h
+++ b/Engine/Plugins/VirtualProduction/DMX/DMXPixelMapping/Source/DMXPixelMappingEditor/Private/DMXPixelMappingEditorStyle.h
@@ -1,30 +1,18 @@
 // Copyright Epic Games, Inc. All Rights Reserved.
 
 #pragma once
-#include "Templates/SharedPointer.h"
 
-class ISlateStyle;
+#include "Styling/SlateStyle.h"
+
 
 class FDMXPixelMappingEditorStyle
+	: public FSlateStyleSet
 {
 public:
-	static void Initialize();
-
-	static void Shutdown();
+	FDMXPixelMappingEditorStyle();
 
-	/** reloads textures used by slate renderer */
-	static void ReloadTextures();
+	virtual ~FDMXPixelMappingEditorStyle();
 
-	/** @return The Slate style set for niagara editor widgets */
+	/** @return The Slate style set for pixel mapping editor widgets */
 	static const ISlateStyle& Get();
-
-	static FName GetStyleSetName();
-
-private:
-
-	static TSharedRef< class FSlateStyleSet > Create();
-
-private:
-
-	static TSharedPtr< class FSlateStyleSet > DMXPixelMappingEditorStyleInstance;
 };
diff --git a/Engine/Plugins/VirtualProduction/DMX/DMXPixelMapping/Source/DMXPixelMappingEditor/Private/DMXPixelMappingToolbar.cpp b/Engine/Plugins/VirtualProduction/DMX/DMXPixelMapping/Source/DMXPixelMappingEditor/Private/DMXPixelMappingToolbar.cpp
--- a/Engine/Plugins/VirtualProduction/DMX/DMXPixelMapping/Source/DMXPixelMappingEditor/Private/DMXPixelMappingToolbar.cpp
+++ b/Engine/Plugins/VirtualProduction/DMX/DMXPixelMapping/Source/DMXPixelMappingEditor/Private/DMXPixelMappingToolbar.cpp
@@ -35,7 +35,7 @@ void FDMXPixelMappingToolbar::BuildToolbarCallback(FToolBarBuilder& ToolbarBuild
 			NAME_None, 
 			TAttribute<FText>(), 
 			TAttribute<FText>(),
-			FSlateIcon(FDMXPixelMappingEditorStyle::GetStyleSetName(), "DMXPixelMappingEditor.AddMapping"),
+			FSlateIcon(FDMXPixelMappingEditorStyle::Get().GetStyleSetName(), "Icons.AddMapping"),
 			FName(TEXT("Add Source")));
 	}
 	ToolbarBuilder.EndSection();
@@ -46,14 +46,14 @@ void FDMXPixelMappingToolbar::BuildToolbarCallback(FToolBarBuilder& ToolbarBuild
 			NAME_None, 
 			TAttribute<FText>(), 
 			TAttribute<FText>(),
-			FSlateIcon(FDMXPixelMappingEditorStyle::GetStyleSetName(), "DMXPixelMappingEditor.PlayDMX"),
+			FSlateIcon(FDMXPixelMappingEditorStyle::Get().GetStyleSetName(), "Icons.PlayDMX"),
 			FName(TEXT("Play DMX")));
 
 		ToolbarBuilder.AddToolBarButton(FDMXPixelMappingEditorCommands::Get().StopPlayingDMX,
 			NAME_None, 
 			TAttribute<FText>(), 
 			TAttribute<FText>(),
-			FSlateIcon(FDMXPixelMappingEditorStyle::GetStyleSetName(), "DMXPixelMappingEditor.StopPlayingDMX"),
+			FSlateIcon(FDMXPixelMappingEditorStyle::Get().GetStyleSetName(), "Icons.StopPlayingDMX"),
 			FName(TEXT("Stop Playing DMX")));
 	}
 	ToolbarBuilder.EndSection();
diff --git a/Engine/Plugins/VirtualProduction/DMX/DMXPixelMapping/Source/DMXPixelMappingEditor/Private/Toolkits/DMXPixelMappingToolkit.cpp b/Engine/Plugins/VirtualProduction/DMX/DMXPixelMapping/Source/DMXPixelMappingEditor/Private/Toolkits/DMXPixelMappingToolkit.cpp
--- a/Engine/Plugins/VirtualProduction/DMX/DMXPixelMapping/Source/DMXPixelMappingEditor/Private/Toolkits/DMXPixelMappingToolkit.cpp
+++ b/Engine/Plugins/VirtualProduction/DMX/DMXPixelMapping/Source/DMXPixelMappingEditor/Private/Toolkits/DMXPixelMappingToolkit.cpp
@@ -72,7 +72,7 @@ void FDMXPixelMappingToolkit::RegisterTabSpawners(const TSharedRef<class FTabMan
 	InTabManager->RegisterTabSpawner(DMXLibraryViewTabID, FOnSpawnTab::CreateSP(this, &FDMXPixelMappingToolkit::SpawnTab_DMXLibraryView))
 		.SetDisplayName(LOCTEXT("Tab_DMXLibraryView", "DMX Library"))
 		.SetGroup(WorkspaceMenuCategoryRef)
-		.SetIcon(FSlateIcon(FDMXPixelMappingEditorStyle::GetStyleSetName(), "ClassIcon.DMXPixelMapping"));
+		.SetIcon(FSlateIcon(FDMXPixelMappingEditorStyle::Get().GetStyleSetName(), "ClassIcon.DMXPixelMapping"));
 
 	InTabManager->RegisterTabSpawner(HierarchyViewTabID, FOnSpawnTab::CreateSP(this, &FDMXPixelMappingToolkit::SpawnTab_HierarchyView))
 		.SetDisplayName(LOCTEXT("Tab_HierarchyView", "Hierarchy"))
@@ -87,7 +87,7 @@ void FDMXPixelMappingToolkit::RegisterTabSpawners(const TSharedRef<class FTabMan
 	InTabManager->RegisterTabSpawner(PreviewViewTabID, FOnSpawnTab::CreateSP(this, &FDMXPixelMappingToolkit::SpawnTab_PreviewView))
 		.SetDisplayName(LOCTEXT("Tab_PreviewView", "Preview"))
 		.SetGroup(WorkspaceMenuCategoryRef)
-		.SetIcon(FSlateIcon(FDMXPixelMappingEditorStyle::GetStyleSetName(), "Icons.Preview"));
+		.SetIcon(FSlateIcon(FDMXPixelMappingEditorStyle::Get().GetStyleSetName(), "Icons.Preview"));
 
 	InTabManager->RegisterTabSpawner(DetailsViewTabID, FOnSpawnTab::CreateSP(this, &FDMXPixelMappingToolkit::SpawnTab_DetailsView))
 		.SetDisplayName(LOCTEXT("Tab_DetailsView", "Details"))
diff --git a/Engine/Plugins/VirtualProduction/DMX/DMXPixelMapping/Source/DMXPixelMappingEditor/Private/Views/SDMXPixelMappingDesignerView.cpp b/Engine/Plugins/VirtualProduction/DMX/DMXPixelMapping/Source/DMXPixelMappingEditor/Private/Views/SDMXPixelMappingDesignerView.cpp
--- a/Engine/Plugins/VirtualProduction/DMX/DMXPixelMapping/Source/DMXPixelMappingEditor/Private/Views/SDMXPixelMappingDesignerView.cpp
+++ b/Engine/Plugins/VirtualProduction/DMX/DMXPixelMapping/Source/DMXPixelMappingEditor/Private/Views/SDMXPixelMappingDesignerView.cpp
@@ -208,7 +208,6 @@ void SDMXPixelMappingDesignerView::Construct(const FArguments& InArgs, const TSh
 							[
 								SNew(SButton)
 								.ButtonColorAndOpacity(FLinearColor::Transparent)
-								.ButtonStyle(FAppStyle::Get(), "ViewportMenu.Button")
 								.ToolTipText(LOCTEXT("ZoomToFit_ToolTip", "Zoom To Fit"))
 								.OnClicked(this, &SDMXPixelMappingDesignerView::OnZoomToFitClicked)
 								.ContentPadding(FAppStyle::Get().GetMargin("ViewportMenu.SToolBarButtonBlock.Button.Padding"))
diff --git a/Engine/Plugins/VirtualProduction/DMX/DMXPixelMapping/Source/DMXPixelMappingEditor/Private/Views/SDMXPixelMappingPreviewView.cpp b/Engine/Plugins/VirtualProduction/DMX/DMXPixelMapping/Source/DMXPixelMappingEditor/Private/Views/SDMXPixelMappingPreviewView.cpp
--- a/Engine/Plugins/VirtualProduction/DMX/DMXPixelMapping/Source/DMXPixelMappingEditor/Private/Views/SDMXPixelMappingPreviewView.cpp
+++ b/Engine/Plugins/VirtualProduction/DMX/DMXPixelMapping/Source/DMXPixelMappingEditor/Private/Views/SDMXPixelMappingPreviewView.cpp
@@ -289,7 +289,6 @@ TSharedRef<SWidget> SDMXPixelMappingPreviewView::CreateOverlayUI()
 			[
 				SNew(SButton)
 				.ButtonColorAndOpacity(FLinearColor::Transparent)
-				.ButtonStyle(FAppStyle::Get(), "ViewportMenu.Button")
 				.ToolTipText(LOCTEXT("ZoomToFit_ToolTip", "Zoom To Fit"))
 				.OnClicked(this, &SDMXPixelMappingPreviewView::HandleZoomToFitClicked)
 				.ContentPadding(FAppStyle::Get().GetMargin("ViewportMenu.SToolBarButtonBlock.Button.Padding"))
