DMX - Fix CIS  Static Analysis Win64

diff --git a/Engine/Plugins/VirtualProduction/DMX/DMXPixelMapping/Source/DMXPixelMappingEditor/Private/ViewModels/DMXPixelMappingDMXLibraryViewModel.cpp b/Engine/Plugins/VirtualProduction/DMX/DMXPixelMapping/Source/DMXPixelMappingEditor/Private/ViewModels/DMXPixelMappingDMXLibraryViewModel.cpp
--- a/Engine/Plugins/VirtualProduction/DMX/DMXPixelMapping/Source/DMXPixelMappingEditor/Private/ViewModels/DMXPixelMappingDMXLibraryViewModel.cpp
+++ b/Engine/Plugins/VirtualProduction/DMX/DMXPixelMapping/Source/DMXPixelMappingEditor/Private/ViewModels/DMXPixelMappingDMXLibraryViewModel.cpp
@@ -135,9 +135,9 @@ void UDMXPixelMappingDMXLibraryViewModel::AddFixturePatchesEnsured(const TArray<
 	// Ensure all patches are of same library
 	const UDMXEntityFixturePatch* const * FixturePatchOfDifferentLibraryPtr = Algo::FindByPredicate(FixturePatches, [CommonDMXLibrary](const UDMXEntityFixturePatch* FixturePatch)
 		{
-			return FixturePatch && FixturePatch->GetParentLibrary() == CommonDMXLibrary;
+			return FixturePatch && FixturePatch->GetParentLibrary() != CommonDMXLibrary;
 		});
-	if (!ensureMsgf(FixturePatchOfDifferentLibraryPtr, TEXT("Cannot add Fixture Patches to Pixel Mapping. Patches don't share a common library")))
+	if (!ensureMsgf(!FixturePatchOfDifferentLibraryPtr, TEXT("Cannot add Fixture Patches to Pixel Mapping. Patches don't share a common library")))
 	{
 		const FString PreviousDMXLibraryName = CommonDMXLibrary->GetName();
 		const FString FixturePatchOfOtherLibraryName = (*FixturePatchOfDifferentLibraryPtr) ? (*FixturePatchOfDifferentLibraryPtr)->GetName() : TEXT("Invalid Fixture Patch");
