DMX - Change 'Add Source' icon, remove obsolete style resources

diff --git a/Engine/Plugins/VirtualProduction/DMX/DMXPixelMapping/Source/DMXPixelMappingEditor/Private/DMXPixelMappingEditorStyle.cpp b/Engine/Plugins/VirtualProduction/DMX/DMXPixelMapping/Source/DMXPixelMappingEditor/Private/DMXPixelMappingEditorStyle.cpp
--- a/Engine/Plugins/VirtualProduction/DMX/DMXPixelMapping/Source/DMXPixelMappingEditor/Private/DMXPixelMappingEditorStyle.cpp
+++ b/Engine/Plugins/VirtualProduction/DMX/DMXPixelMapping/Source/DMXPixelMappingEditor/Private/DMXPixelMappingEditorStyle.cpp
@@ -8,10 +8,10 @@
 #include "Misc/Paths.h"
 #include "Rendering/SlateRenderer.h"
 #include "Styling/SlateBrush.h"
-#include "Styling/StyleColors.h"
+#include "Styling/SlateStyle.h"
 #include "Styling/SlateStyleMacros.h"
 #include "Styling/SlateStyleRegistry.h"
-#include "Styling/SlateStyle.h"
+#include "Styling/StyleColors.h"
 
 
 FDMXPixelMappingEditorStyle::FDMXPixelMappingEditorStyle()
@@ -41,7 +41,7 @@ FDMXPixelMappingEditorStyle::FDMXPixelMappingEditorStyle()
 	{
 		Set("Icons.Preview", new IMAGE_BRUSH_SVG("Preview", Icon16x16));
 		Set("Icons.ZoomToFit", new IMAGE_BRUSH_SVG("ZoomToFit", Icon16x16));
-		Set("Icons.AddMapping", new IMAGE_BRUSH("icon_DMXPixelMappingEditor_AddMapping_40x", Icon16x16));
+		Set("Icons.AddSource",new CORE_IMAGE_BRUSH("Icons/PlusSymbol_12x", Icon12x12));
 		Set("Icons.PlayDMX", new CORE_IMAGE_BRUSH_SVG("Starship/Common/play", Icon16x16, FStyleColors::AccentGreen));
 		Set("Icons.StopPlayingDMX", new CORE_IMAGE_BRUSH_SVG("Starship/Common/stop", Icon16x16, FStyleColors::AccentRed));
 	}
diff --git a/Engine/Plugins/VirtualProduction/DMX/DMXPixelMapping/Source/DMXPixelMappingEditor/Private/DMXPixelMappingToolbar.cpp b/Engine/Plugins/VirtualProduction/DMX/DMXPixelMapping/Source/DMXPixelMappingEditor/Private/DMXPixelMappingToolbar.cpp
--- a/Engine/Plugins/VirtualProduction/DMX/DMXPixelMapping/Source/DMXPixelMappingEditor/Private/DMXPixelMappingToolbar.cpp
+++ b/Engine/Plugins/VirtualProduction/DMX/DMXPixelMapping/Source/DMXPixelMappingEditor/Private/DMXPixelMappingToolbar.cpp
@@ -1,12 +1,14 @@
 // Copyright Epic Games, Inc. All Rights Reserved.
 
 #include "DMXPixelMappingToolbar.h"
-#include "Framework/Commands/UICommandList.h"
-#include "Toolkits/DMXPixelMappingToolkit.h"
+
 #include "DMXPixelMappingEditorCommands.h"
 #include "DMXPixelMappingEditorStyle.h"
-
 #include "Framework/MultiBox/MultiBoxBuilder.h"
+#include "Framework/Commands/UICommandList.h"
+#include "Styling/AppStyle.h"
+#include "Toolkits/DMXPixelMappingToolkit.h"
+
 
 #define LOCTEXT_NAMESPACE "FDMXPixelMappingToolbar"
 
@@ -35,7 +37,7 @@ void FDMXPixelMappingToolbar::BuildToolbarCallback(FToolBarBuilder& ToolbarBuild
 			NAME_None, 
 			TAttribute<FText>(), 
 			TAttribute<FText>(),
-			FSlateIcon(FDMXPixelMappingEditorStyle::Get().GetStyleSetName(), "Icons.AddMapping"),
+			FSlateIcon(FDMXPixelMappingEditorStyle::Get().GetStyleSetName(), "Icons.AddSource"),
 			FName(TEXT("Add Source")));
 	}
 	ToolbarBuilder.EndSection();
