DMX - Add Pixel Mapping grid snapping

diff --git a/Engine/Plugins/VirtualProduction/DMX/DMXPixelMapping/Source/DMXPixelMappingEditor/DMXPixelMappingEditor.Build.cs b/Engine/Plugins/VirtualProduction/DMX/DMXPixelMapping/Source/DMXPixelMappingEditor/DMXPixelMappingEditor.Build.cs
--- a/Engine/Plugins/VirtualProduction/DMX/DMXPixelMapping/Source/DMXPixelMappingEditor/DMXPixelMappingEditor.Build.cs
+++ b/Engine/Plugins/VirtualProduction/DMX/DMXPixelMapping/Source/DMXPixelMappingEditor/DMXPixelMappingEditor.Build.cs
@@ -10,6 +10,7 @@ public class DMXPixelMappingEditor : ModuleRules
 		var EngineDir = Path.GetFullPath(Target.RelativeEnginePath);
 
 		PrivateDependencyModuleNames.AddRange( new string[] {
+			"AppFramework",
 			"ApplicationCore",
 			"Core",
 			"CoreUObject",
diff --git a/Engine/Plugins/VirtualProduction/DMX/DMXPixelMapping/Source/DMXPixelMappingEditor/Private/Customizations/DMXPixelMappingDetailCustomization_FixtureGroup.cpp b/Engine/Plugins/VirtualProduction/DMX/DMXPixelMapping/Source/DMXPixelMappingEditor/Private/Customizations/DMXPixelMappingDetailCustomization_FixtureGroup.cpp
--- a/Engine/Plugins/VirtualProduction/DMX/DMXPixelMapping/Source/DMXPixelMappingEditor/Private/Customizations/DMXPixelMappingDetailCustomization_FixtureGroup.cpp
+++ b/Engine/Plugins/VirtualProduction/DMX/DMXPixelMapping/Source/DMXPixelMappingEditor/Private/Customizations/DMXPixelMappingDetailCustomization_FixtureGroup.cpp
@@ -129,13 +129,13 @@ void FDMXPixelMappingDetailCustomization_FixtureGroup::HandleSizePropertyChanged
 			{
 				Child->Modify();
 
-				// Scale size (SetSize already clamps)
-				Child->SetSize(Child->GetSize() * RatioVector);
+				// Scale size
+				const FVector2D NewChildSize = Child->GetSize() * RatioVector;
 
 				// Scale position
-				const FVector2D ChildPosition = Child->GetPosition();
-				const FVector2D NewPositionRelative = (ChildPosition - GroupPosition) * RatioVector;
-				Child->SetPosition(GroupPosition + NewPositionRelative);
+				const FVector2D NewChildPosition = Child->GetPosition();
+				const FVector2D NewChildPositionRelative = (NewChildPosition - GroupPosition) * RatioVector;
+				Child->SetPosition(GroupPosition + NewChildPositionRelative);
 			}
 		}
 	}
diff --git a/Engine/Plugins/VirtualProduction/DMX/DMXPixelMapping/Source/DMXPixelMappingEditor/Private/Customizations/DMXPixelMappingDetailCustomization_Renderer.cpp b/Engine/Plugins/VirtualProduction/DMX/DMXPixelMapping/Source/DMXPixelMappingEditor/Private/Customizations/DMXPixelMappingDetailCustomization_Renderer.cpp
--- a/Engine/Plugins/VirtualProduction/DMX/DMXPixelMapping/Source/DMXPixelMappingEditor/Private/Customizations/DMXPixelMappingDetailCustomization_Renderer.cpp
+++ b/Engine/Plugins/VirtualProduction/DMX/DMXPixelMapping/Source/DMXPixelMappingEditor/Private/Customizations/DMXPixelMappingDetailCustomization_Renderer.cpp
@@ -91,30 +91,6 @@ void FDMXPixelMappingDetailCustomization_Renderer::CustomizeDetails(IDetailLayou
 	}
 }
 
-void FDMXPixelMappingDetailCustomization_Renderer::PostUndo(bool bSuccess)
-{
-	const TArray<TWeakObjectPtr<UObject>> SelectedObjects = PropertyUtilities->GetSelectedObjects();
-	for (const TWeakObjectPtr<UObject>& Object : SelectedObjects)
-	{
-		if (UDMXPixelMappingRendererComponent* RendererComponent = Cast<UDMXPixelMappingRendererComponent>(Object))
-		{
-			RendererComponent->UpdatePreprocessRenderer();
-		}
-	}
-}
-
-void FDMXPixelMappingDetailCustomization_Renderer::PostRedo(bool bSuccess)
-{
-	const TArray<TWeakObjectPtr<UObject>> SelectedObjects = PropertyUtilities->GetSelectedObjects();
-	for (const TWeakObjectPtr<UObject>& Object : SelectedObjects)
-	{
-		if (UDMXPixelMappingRendererComponent* RendererComponent = Cast<UDMXPixelMappingRendererComponent>(Object))
-		{
-			RendererComponent->UpdatePreprocessRenderer();
-		}
-	}
-}
-
 void FDMXPixelMappingDetailCustomization_Renderer::AddInputTextureWarning(IDetailCategoryBuilder& InCategory)
 {
 	const FSlateBrush* WarningIcon = FAppStyle::GetBrush("SettingsEditor.WarningIcon");
diff --git a/Engine/Plugins/VirtualProduction/DMX/DMXPixelMapping/Source/DMXPixelMappingEditor/Private/Customizations/DMXPixelMappingDetailCustomization_Renderer.h b/Engine/Plugins/VirtualProduction/DMX/DMXPixelMapping/Source/DMXPixelMappingEditor/Private/Customizations/DMXPixelMappingDetailCustomization_Renderer.h
--- a/Engine/Plugins/VirtualProduction/DMX/DMXPixelMapping/Source/DMXPixelMappingEditor/Private/Customizations/DMXPixelMappingDetailCustomization_Renderer.h
+++ b/Engine/Plugins/VirtualProduction/DMX/DMXPixelMapping/Source/DMXPixelMappingEditor/Private/Customizations/DMXPixelMappingDetailCustomization_Renderer.h
@@ -2,7 +2,6 @@
 
 #pragma once
 
-#include "EditorUndoClient.h"
 #include "IDetailCustomization.h"
 #include "UObject/WeakObjectPtr.h"
 
@@ -18,7 +17,6 @@ class UDMXPixelMappingRendererComponent;
 
 class FDMXPixelMappingDetailCustomization_Renderer
 	: public IDetailCustomization
-	, public FSelfRegisteringEditorUndoClient
 {
 public:
 	FDMXPixelMappingDetailCustomization_Renderer(const TWeakPtr<FDMXPixelMappingToolkit> InWeakToolkit);
@@ -31,12 +29,6 @@ public:
 	/** IDetailCustomization interface */
 	virtual void CustomizeDetails(IDetailLayoutBuilder& DetailLayout) override;
 
-protected:
-	//~ Begin FSelfRegisteringEditorUndoClient interface
-	virtual void PostUndo(bool bSuccess) override;
-	virtual void PostRedo(bool bSuccess) override;
-	//~ End FSelfRegisteringEditorUndoClient interface
-
 private:
 	/** Adds a warning if the input texture is not set */
 	void AddInputTextureWarning(IDetailCategoryBuilder& InCategory);
diff --git a/Engine/Plugins/VirtualProduction/DMX/DMXPixelMapping/Source/DMXPixelMappingEditor/Private/DMXPixelMappingEditorCommands.cpp b/Engine/Plugins/VirtualProduction/DMX/DMXPixelMapping/Source/DMXPixelMappingEditor/Private/DMXPixelMappingEditorCommands.cpp
--- a/Engine/Plugins/VirtualProduction/DMX/DMXPixelMapping/Source/DMXPixelMappingEditor/Private/DMXPixelMappingEditorCommands.cpp
+++ b/Engine/Plugins/VirtualProduction/DMX/DMXPixelMapping/Source/DMXPixelMappingEditor/Private/DMXPixelMappingEditorCommands.cpp
@@ -19,13 +19,13 @@ FDMXPixelMappingEditorCommands::FDMXPixelMappingEditorCommands()
 
 void FDMXPixelMappingEditorCommands::RegisterCommands()
 {
-	UI_COMMAND(SaveThumbnailImage, "Thumbnail", "Generate Thumbnail", EUserInterfaceActionType::Button, FInputChord());
-
 	UI_COMMAND(AddMapping, "Add Source", "Adds a new Source Texure, Material or User Widget to the Pixel Map asset.", EUserInterfaceActionType::Button, FInputChord());
 
 	UI_COMMAND(PlayDMX, "Play DMX", "Starts to send DMX", EUserInterfaceActionType::Button, FInputChord());
 	UI_COMMAND(StopPlayingDMX, "Stop Playing DMX", "Stops playing DMX.", EUserInterfaceActionType::Button, FInputChord());
 
+	UI_COMMAND(ToggleGridSnapping, "Toggle Grid Snapping", "Enables/disables grid snapping", EUserInterfaceActionType::Button, FInputChord());
+
 	UI_COMMAND(SizeComponentToTexture, "Size Component to Texture", "Sizes the selected group to Texture.", EUserInterfaceActionType::Button, FInputChord(EModifierKey::Shift, EKeys::S));
 	
 	UI_COMMAND(ToggleScaleChildrenWithParent, "Scale Children with Parent", "Sets if children are scaled with their parent, when the parent is resized.", EUserInterfaceActionType::ToggleButton, FInputChord(EModifierKey::Alt, EKeys::Q));
diff --git a/Engine/Plugins/VirtualProduction/DMX/DMXPixelMapping/Source/DMXPixelMappingEditor/Private/DMXPixelMappingEditorCommands.h b/Engine/Plugins/VirtualProduction/DMX/DMXPixelMapping/Source/DMXPixelMappingEditor/Private/DMXPixelMappingEditorCommands.h
--- a/Engine/Plugins/VirtualProduction/DMX/DMXPixelMapping/Source/DMXPixelMappingEditor/Private/DMXPixelMappingEditorCommands.h
+++ b/Engine/Plugins/VirtualProduction/DMX/DMXPixelMapping/Source/DMXPixelMappingEditor/Private/DMXPixelMappingEditorCommands.h
@@ -14,11 +14,13 @@ public:
 
 	virtual void RegisterCommands() override;
 
-	TSharedPtr<FUICommandInfo> SaveThumbnailImage;
 	TSharedPtr<FUICommandInfo> AddMapping;
 	TSharedPtr<FUICommandInfo> PlayDMX;
 	TSharedPtr<FUICommandInfo> StopPlayingDMX;
 
+	// Grid snapping related
+	TSharedPtr<FUICommandInfo> ToggleGridSnapping;
+
 	// Layout related
 	TSharedPtr<FUICommandInfo> SizeComponentToTexture;
 
diff --git a/Engine/Plugins/VirtualProduction/DMX/DMXPixelMapping/Source/DMXPixelMappingEditor/Private/DMXPixelMappingEditorStyle.cpp b/Engine/Plugins/VirtualProduction/DMX/DMXPixelMapping/Source/DMXPixelMappingEditor/Private/DMXPixelMappingEditorStyle.cpp
--- a/Engine/Plugins/VirtualProduction/DMX/DMXPixelMapping/Source/DMXPixelMappingEditor/Private/DMXPixelMappingEditorStyle.cpp
+++ b/Engine/Plugins/VirtualProduction/DMX/DMXPixelMapping/Source/DMXPixelMappingEditor/Private/DMXPixelMappingEditorStyle.cpp
@@ -40,7 +40,6 @@ FDMXPixelMappingEditorStyle::FDMXPixelMappingEditorStyle()
 	// Icons
 	{
 		Set("Icons.Preview", new IMAGE_BRUSH_SVG("Preview", Icon16x16));
-		Set("Icons.ZoomToFit", new IMAGE_BRUSH_SVG("ZoomToFit", Icon16x16));
 		Set("Icons.AddSource",new CORE_IMAGE_BRUSH("Icons/PlusSymbol_12x", Icon12x12));
 		Set("Icons.PlayDMX", new CORE_IMAGE_BRUSH_SVG("Starship/Common/play", Icon16x16, FStyleColors::AccentGreen));
 		Set("Icons.StopPlayingDMX", new CORE_IMAGE_BRUSH_SVG("Starship/Common/stop", Icon16x16, FStyleColors::AccentRed));
diff --git a/Engine/Plugins/VirtualProduction/DMX/DMXPixelMapping/Source/DMXPixelMappingEditor/Private/DragDrop/DMXPixelMappingDragDropOp.cpp b/Engine/Plugins/VirtualProduction/DMX/DMXPixelMapping/Source/DMXPixelMappingEditor/Private/DragDrop/DMXPixelMappingDragDropOp.cpp
--- a/Engine/Plugins/VirtualProduction/DMX/DMXPixelMapping/Source/DMXPixelMappingEditor/Private/DragDrop/DMXPixelMappingDragDropOp.cpp
+++ b/Engine/Plugins/VirtualProduction/DMX/DMXPixelMapping/Source/DMXPixelMappingEditor/Private/DragDrop/DMXPixelMappingDragDropOp.cpp
@@ -2,11 +2,12 @@
 
 #include "DragDrop/DMXPixelMappingDragDropOp.h"
 
-#include "DragDrop/DMXPixelMappingDragDropOp.h"
-#include "DragDrop/DMXPixelMappingGroupChildDragDropHelper.h"
 #include "Components/DMXPixelMappingMatrixComponent.h"
-
+#include "Components/DMXPixelMappingRendererComponent.h"
+#include "DMXPixelMapping.h"
+#include "DragDrop/DMXPixelMappingGroupChildDragDropHelper.h"
 #include "Editor.h"
+#include "Toolkits/DMXPixelMappingToolkit.h"
 
 
 #define LOCTEXT_NAMESPACE "FDMXPixelMappingDragDropOp"
@@ -17,10 +18,11 @@ FDMXPixelMappingDragDropOp::~FDMXPixelMappingDragDropOp()
 	GEditor->EndTransaction();
 }
 
-TSharedRef<FDMXPixelMappingDragDropOp> FDMXPixelMappingDragDropOp::New(const FVector2D& InGraphSpaceDragOffset, const TArray<TSharedPtr<FDMXPixelMappingComponentTemplate>>& InTemplates, UDMXPixelMappingBaseComponent* InParent)
+TSharedRef<FDMXPixelMappingDragDropOp> FDMXPixelMappingDragDropOp::New(const TSharedRef<FDMXPixelMappingToolkit>& InToolkit, const FVector2D& InGraphSpaceDragOffset, const TArray<TSharedPtr<FDMXPixelMappingComponentTemplate>>& InTemplates, UDMXPixelMappingBaseComponent* InParent)
 {
 	TSharedRef<FDMXPixelMappingDragDropOp> Operation = MakeShared<FDMXPixelMappingDragDropOp>();
-
+	
+	Operation->WeakToolkit = InToolkit;
 	Operation->Templates = InTemplates;
 	Operation->bWasCreatedAsTemplate = true;
 	Operation->Parent = InParent;
@@ -35,10 +37,11 @@ TSharedRef<FDMXPixelMappingDragDropOp> FDMXPixelMappingDragDropOp::New(const FVe
 	return Operation;
 }
 
-TSharedRef<FDMXPixelMappingDragDropOp> FDMXPixelMappingDragDropOp::New(const FVector2D& InGraphSpaceDragOffset, const TArray<TWeakObjectPtr<UDMXPixelMappingBaseComponent>>& InDraggedComponents)
+TSharedRef<FDMXPixelMappingDragDropOp> FDMXPixelMappingDragDropOp::New(const TSharedRef<FDMXPixelMappingToolkit>& InToolkit, const FVector2D& InGraphSpaceDragOffset, const TArray<TWeakObjectPtr<UDMXPixelMappingBaseComponent>>& InDraggedComponents)
 {
 	TSharedRef<FDMXPixelMappingDragDropOp> Operation = MakeShared<FDMXPixelMappingDragDropOp>();
 
+	Operation->WeakToolkit = InToolkit;
 	Operation->bWasCreatedAsTemplate = false;
 	Operation->SetDraggedComponents(InDraggedComponents);
 	Operation->GraphSpaceDragOffset = InGraphSpaceDragOffset;
@@ -58,44 +61,98 @@ void FDMXPixelMappingDragDropOp::SetDraggedComponents(const TArray<TWeakObjectPt
 	DraggedComponents = InDraggedComponents;
 	Templates.Reset();
 
-	// Rebuild the group child drag drop helper in case that's what was set
+	// Rebuild the group child drag drop helper
 	GroupChildDragDropHelper = FDMXPixelMappingGroupChildDragDropHelper::Create(AsShared());
 }
 
 void FDMXPixelMappingDragDropOp::LayoutOutputComponents(const FVector2D& GraphSpacePosition)
 {
-	if (DraggedComponents.Num() > 0)
+	UDMXPixelMappingOutputComponent* FirstComponent = DraggedComponents.IsEmpty() ? nullptr : Cast<UDMXPixelMappingOutputComponent>(DraggedComponents[0]);
+	if (!FirstComponent)
 	{
-		if (UDMXPixelMappingOutputComponent* FirstComponent = Cast<UDMXPixelMappingOutputComponent>(DraggedComponents[0]))
+		return;
+	}
+	const FVector2D Anchor = FirstComponent->GetPosition();
+
+	// Move all to new position
+	for (const TWeakObjectPtr<UDMXPixelMappingBaseComponent>& Component : DraggedComponents)
+	{
+		if (UDMXPixelMappingOutputComponent* OutputComponent = Cast<UDMXPixelMappingOutputComponent>(Component.Get()))
 		{
-			const FVector2D Anchor = FirstComponent->GetPosition();
+			OutputComponent->PreEditChange(nullptr);
+			
+			constexpr bool bModifyChildrenRecursive = true;
+			Component->ForEachChild([](UDMXPixelMappingBaseComponent* Component)
+				{
+					Component->Modify();
+				}, bModifyChildrenRecursive);
 
-			// Move all to new position
-			for (const TWeakObjectPtr<UDMXPixelMappingBaseComponent>& Component : DraggedComponents)
+			if (ensureMsgf(OutputComponent->GetClass() != UDMXPixelMappingMatrixComponent::StaticClass(),
+				TEXT("Only matrix components can laid out with FDMXPixelMappingDragDropOp::LayoutOutputComponents. Please use FGroupChildDragDropHelper instead (see FDMXPixelMappingDragDropOp::GetGroupChildDragDropHelper().")))
 			{
-				if (UDMXPixelMappingOutputComponent* OutputComponent = Cast<UDMXPixelMappingOutputComponent>(Component.Get()))
-				{
-					OutputComponent->PreEditChange(nullptr);
+				const FVector2D AnchorOffset = Anchor - OutputComponent->GetPosition();
 
-					constexpr bool bModifyChildrenRecursive = true;
-					Component->ForEachChild([](UDMXPixelMappingBaseComponent* Component)
-						{
-							Component->Modify();
-						}, bModifyChildrenRecursive);
+				const FVector2D NewPosition = FVector2D(GraphSpacePosition - AnchorOffset - GraphSpaceDragOffset).RoundToVector();
+				OutputComponent->SetPosition(NewPosition);
+			}
+		}
+	}
+	GridSnap();
 
-					if (ensureMsgf(OutputComponent->GetClass() != UDMXPixelMappingMatrixComponent::StaticClass(),
-						TEXT("Matrix Cells are not supported. Use the GroupChildDragDropHelper from this class instead")))
-					{
-						FVector2D AnchorOffset = Anchor - OutputComponent->GetPosition();
+	for (const TWeakObjectPtr<UDMXPixelMappingBaseComponent>& Component : DraggedComponents)
+	{
+		if (UDMXPixelMappingOutputComponent* OutputComponent = Cast<UDMXPixelMappingOutputComponent>(Component.Get()))
+		{
+			OutputComponent->PostEditChange();
+		}
+	}
+}
+
+void FDMXPixelMappingDragDropOp::GridSnap()
+{
+	const UDMXPixelMapping* PixelMapping = WeakToolkit.IsValid() ? WeakToolkit.Pin()->GetDMXPixelMapping() : nullptr;
+	if (!PixelMapping || DraggedComponents.IsEmpty() || PixelMapping->SnapGridColumns == 0 || PixelMapping->SnapGridRows == 0)
+	{
+		return;
+	}
 
-						const FVector2D NewPosition = GraphSpacePosition - AnchorOffset - GraphSpaceDragOffset;
-						OutputComponent->SetPosition(NewPosition.RoundToVector());
+	// Grid snap the first component only, and move the others by same delta.
+	UDMXPixelMappingOutputComponent* FirstComponent = Cast<UDMXPixelMappingOutputComponent>(DraggedComponents[0].Get());
+	if (!FirstComponent)
+	{
+		return;
+	}
 
-					}
+	UDMXPixelMappingRendererComponent* RendererOfFirstComponent = FirstComponent->GetRendererComponent();
+	if (!RendererOfFirstComponent)
+	{
+		return;
+	}
 
-					OutputComponent->PostEditChange();
-				}
+	const FVector2D CellSize = [PixelMapping, RendererOfFirstComponent]()
+		{
+			if (PixelMapping->bGridSnappingEnabled)
+			{
+				const FVector2D TextureSize = RendererOfFirstComponent->GetSize();
+				return TextureSize / FVector2D(PixelMapping->SnapGridColumns, PixelMapping->SnapGridRows);
 			}
+			
+			// Grid snap to pixels if grid snapping is disabled
+			return FVector2D(1.f, 1.f);
+		}();
+
+	const int32 Column = FirstComponent->GetPosition().X / CellSize.X;
+	const int32 Row = FirstComponent->GetPosition().Y / CellSize.Y;
+
+	const FVector2D GridSnapPosition = FVector2D(Column, Row) * CellSize;
+	const FVector2D DeltaVector = GridSnapPosition - FirstComponent->GetPosition();
+
+	FirstComponent->SetPosition(GridSnapPosition);
+	for (int32 ComponentIndex = 1; ComponentIndex < DraggedComponents.Num(); ComponentIndex++)
+	{
+		if (UDMXPixelMappingOutputComponent* OutputComponent = Cast<UDMXPixelMappingOutputComponent>(DraggedComponents[ComponentIndex].Get()))
+		{
+			OutputComponent->SetPosition(OutputComponent->GetPosition() + DeltaVector);
 		}
 	}
 }
diff --git a/Engine/Plugins/VirtualProduction/DMX/DMXPixelMapping/Source/DMXPixelMappingEditor/Private/DragDrop/DMXPixelMappingDragDropOp.h b/Engine/Plugins/VirtualProduction/DMX/DMXPixelMapping/Source/DMXPixelMappingEditor/Private/DragDrop/DMXPixelMappingDragDropOp.h
--- a/Engine/Plugins/VirtualProduction/DMX/DMXPixelMapping/Source/DMXPixelMappingEditor/Private/DragDrop/DMXPixelMappingDragDropOp.h
+++ b/Engine/Plugins/VirtualProduction/DMX/DMXPixelMapping/Source/DMXPixelMappingEditor/Private/DragDrop/DMXPixelMappingDragDropOp.h
@@ -2,15 +2,12 @@
 
 #pragma once
 
-
-
 #include "DragAndDrop/DecoratedDragDropOp.h"
 
-class FDMXPixelMappingGroupChildDragDropHelper;
-
 class FDMXPixelMappingComponentTemplate;
 class FDMXPixelMappingComponentReference;
-class UDMXPixelMappingBaseComponent;
+class FDMXPixelMappingGroupChildDragDropHelper;
+class FDMXPixelMappingToolkit;
 class UDMXPixelMappingBaseComponent;
 
 
@@ -28,10 +25,10 @@ public:
 	virtual ~FDMXPixelMappingDragDropOp();
 
 	/** Constructs the drag drop operation. */
-	static TSharedRef<FDMXPixelMappingDragDropOp> New(const FVector2D& InGraphSpaceDragOffset, const TArray<TSharedPtr<FDMXPixelMappingComponentTemplate>>& InTemplates, UDMXPixelMappingBaseComponent* InParent);
+	static TSharedRef<FDMXPixelMappingDragDropOp> New(const TSharedRef<FDMXPixelMappingToolkit>& InToolkit, const FVector2D& InGraphSpaceDragOffset, const TArray<TSharedPtr<FDMXPixelMappingComponentTemplate>>& InTemplates, UDMXPixelMappingBaseComponent* InParent);
 	
 	/** Constructs the drag drop operation. Note, the component that initiated the drag drop op has to be the first element in the InDraggedComponents array. */
-	static TSharedRef<FDMXPixelMappingDragDropOp> New(const FVector2D& InGraphSpaceDragOffset, const TArray<TWeakObjectPtr<UDMXPixelMappingBaseComponent>>& InDraggedComponents);
+	static TSharedRef<FDMXPixelMappingDragDropOp> New(const TSharedRef<FDMXPixelMappingToolkit>& InToolkit, const FVector2D& InGraphSpaceDragOffset, const TArray<TWeakObjectPtr<UDMXPixelMappingBaseComponent>>& InDraggedComponents);
 
 	/** Sets dragged components. Clears the template, preventing from the template being used more than once. */
 	void SetDraggedComponents(const TArray<TWeakObjectPtr<UDMXPixelMappingBaseComponent>>& InDraggedComponents);
@@ -45,6 +42,9 @@ public:
 	/** Returns true if the drag drop op was created with the constructor that provides a template */
 	bool WasCreatedAsTemplate() const { return bWasCreatedAsTemplate; }
 
+	/** Utiltity that applies grid snapping. Note, this grid snaps to pixels if grid snapping is disabled. */
+	void GridSnap();
+
 	/** Can be used to store the parent of the dragged component(s) */
 	TWeakObjectPtr<UDMXPixelMappingBaseComponent> Parent;
 
@@ -75,4 +75,7 @@ private:
 
 	/** Group Item drag drop helper */
 	TSharedPtr<FDMXPixelMappingGroupChildDragDropHelper> GroupChildDragDropHelper;
+
+	/** The toolkit that constructed this drag drop op */
+	TWeakPtr<FDMXPixelMappingToolkit> WeakToolkit;
 };
diff --git a/Engine/Plugins/VirtualProduction/DMX/DMXPixelMapping/Source/DMXPixelMappingEditor/Private/DragDrop/DMXPixelMappingGroupChildDragDropHelper.cpp b/Engine/Plugins/VirtualProduction/DMX/DMXPixelMapping/Source/DMXPixelMappingEditor/Private/DragDrop/DMXPixelMappingGroupChildDragDropHelper.cpp
--- a/Engine/Plugins/VirtualProduction/DMX/DMXPixelMapping/Source/DMXPixelMappingEditor/Private/DragDrop/DMXPixelMappingGroupChildDragDropHelper.cpp
+++ b/Engine/Plugins/VirtualProduction/DMX/DMXPixelMapping/Source/DMXPixelMappingEditor/Private/DragDrop/DMXPixelMappingGroupChildDragDropHelper.cpp
@@ -80,104 +80,112 @@ void FDMXPixelMappingGroupChildDragDropHelper::LayoutComponents(const FVector2D&
 
 void FDMXPixelMappingGroupChildDragDropHelper::LayoutAligned(const FVector2D& GraphSpacePosition)
 {
-	if (TSharedPtr<FDMXPixelMappingDragDropOp> DragDropOp = WeakPixelMappingDragDropOp.Pin())
+	TSharedPtr<FDMXPixelMappingDragDropOp> DragDropOp = WeakPixelMappingDragDropOp.Pin();
+	UDMXPixelMappingFixtureGroupComponent* GroupComponent = WeakParentGroupComponent.Get();
+	if (!DragDropOp.IsValid() || !GroupComponent)
 	{
-		if (UDMXPixelMappingFixtureGroupComponent* GroupComponent = WeakParentGroupComponent.Get())
+		return;
+	}
+
+	FVector2D NextPosition = GraphSpacePosition - DragDropOp->GraphSpaceDragOffset;
+	float RowHeight = 0.f;
+
+	for (const TWeakObjectPtr<UDMXPixelMappingOutputComponent>& WeakChildComponent : WeakChildComponents)
+	{
+		if (UDMXPixelMappingOutputComponent* ChildComponent = WeakChildComponent.Get())
 		{
-			FVector2D NextPosition = GraphSpacePosition - DragDropOp->GraphSpaceDragOffset;
-			float RowHeight = 0.f;
+			ChildComponent->PreEditChange(nullptr);
+
+			constexpr bool bModifyChildrenRecursive = true;
+			ChildComponent->ForEachChild([](UDMXPixelMappingBaseComponent* Component)
+				{
+					Component->Modify();
+				}, bModifyChildrenRecursive);
 
-			for (const TWeakObjectPtr<UDMXPixelMappingOutputComponent>& WeakChildComponent : WeakChildComponents)
+			if (GroupComponent->IsOverPosition(NextPosition) && 
+				GroupComponent->IsOverPosition(NextPosition + ChildComponent->GetSize()))
 			{
-				if (UDMXPixelMappingOutputComponent* ChildComponent = WeakChildComponent.Get())
+				ChildComponent->SetPosition(NextPosition);
+
+				RowHeight = FMath::Max(ChildComponent->GetSize().Y, RowHeight);
+				NextPosition = FVector2D(NextPosition.X + ChildComponent->GetSize().X, NextPosition.Y);
+			}
+			else
+			{
+				// Try on a new row
+				FVector2D NewRowPosition = FVector2D(GraphSpacePosition.X, NextPosition.Y + RowHeight);
+
+				const FVector2D NextPositionOnNewRow = FVector2D(NewRowPosition.X + ChildComponent->GetSize().X, NewRowPosition.Y);
+				if (GroupComponent->IsOverPosition(NextPositionOnNewRow) &&
+					GroupComponent->IsOverPosition(NextPositionOnNewRow + ChildComponent->GetSize()))
 				{
-					ChildComponent->PreEditChange(nullptr);
-
-					constexpr bool bModifyChildrenRecursive = true;
-					ChildComponent->ForEachChild([](UDMXPixelMappingBaseComponent* Component)
-						{
-							Component->Modify();
-						}, bModifyChildrenRecursive);
-
-					if (GroupComponent->IsOverPosition(NextPosition) && 
-						GroupComponent->IsOverPosition(NextPosition + ChildComponent->GetSize()))
-					{
-						SetPosition(ChildComponent, NextPosition);
-
-						RowHeight = FMath::Max(ChildComponent->GetSize().Y, RowHeight);
-						NextPosition = FVector2D(NextPosition.X + ChildComponent->GetSize().X, NextPosition.Y);
-					}
-					else
-					{
-						// Try on a new row
-						FVector2D NewRowPosition = FVector2D(GraphSpacePosition.X, NextPosition.Y + RowHeight);
-
-						const FVector2D NextPositionOnNewRow = FVector2D(NewRowPosition.X + ChildComponent->GetSize().X, NewRowPosition.Y);
-						if (GroupComponent->IsOverPosition(NextPositionOnNewRow) &&
-							GroupComponent->IsOverPosition(NextPositionOnNewRow + ChildComponent->GetSize()))
-						{
-							SetPosition(ChildComponent, NewRowPosition);
-
-							NextPosition = FVector2D(NewRowPosition.X + ChildComponent->GetSize().X, NewRowPosition.Y);
-							RowHeight = ChildComponent->GetSize().Y;							
-						}
-						else
-						{
-							SetPosition(ChildComponent, NextPosition);
-
-							NextPosition = FVector2D(NextPosition.X + ChildComponent->GetSize().X, NextPosition.Y);
-						}
-					}
-
-					ChildComponent->PostEditChange();
+					ChildComponent->SetPosition(NewRowPosition);
+
+					NextPosition = FVector2D(NewRowPosition.X + ChildComponent->GetSize().X, NewRowPosition.Y);
+					RowHeight = ChildComponent->GetSize().Y;							
+				}
+				else
+				{
+					ChildComponent->SetPosition(NextPosition);
+
+					NextPosition = FVector2D(NextPosition.X + ChildComponent->GetSize().X, NextPosition.Y);
 				}
 			}
 		}
 	}
+
+	DragDropOp->GridSnap();
+
+	for (const TWeakObjectPtr<UDMXPixelMappingOutputComponent>& WeakChildComponent : WeakChildComponents)
+	{
+		if (UDMXPixelMappingOutputComponent* ChildComponent = WeakChildComponent.Get())
+		{
+			ChildComponent->PostEditChange();
+		}
+	}
 }
 
 void FDMXPixelMappingGroupChildDragDropHelper::LayoutUnaligned(const FVector2D& GraphSpacePosition)
 {
-	if (TSharedPtr<FDMXPixelMappingDragDropOp> PinnedDragDropOp = WeakPixelMappingDragDropOp.Pin())
+	TSharedPtr<FDMXPixelMappingDragDropOp> DragDropOp = WeakPixelMappingDragDropOp.Pin();
+	if (!DragDropOp.IsValid() || WeakChildComponents.IsEmpty())
 	{
-		if (WeakChildComponents.Num() > 0)
-		{
-			if (UDMXPixelMappingOutputComponent* FirstComponent = WeakChildComponents[0].Get())
-			{
-				const FVector2D Anchor = FirstComponent->GetPosition();
+		return;
+	}
 
-				// Move all grou items to their new position
-				for (const TWeakObjectPtr<UDMXPixelMappingOutputComponent>& WeakOutputComponent : WeakChildComponents)
-				{
-					if (UDMXPixelMappingOutputComponent* ChildComponent = WeakOutputComponent.Get())
-					{
-						ChildComponent->PreEditChange(nullptr);
+	UDMXPixelMappingOutputComponent* FirstComponent = WeakChildComponents[0].Get();
+	if (!FirstComponent)
+	{
+		return;
+	}
 
-						constexpr bool bModifyChildrenRecursive = true;
-						ChildComponent->ForEachChild([](UDMXPixelMappingBaseComponent* Component)
-							{
-								Component->Modify();
-							}, bModifyChildrenRecursive);
+	const FVector2D Anchor = FirstComponent->GetPosition();
+	for (const TWeakObjectPtr<UDMXPixelMappingOutputComponent>& WeakOutputComponent : WeakChildComponents)
+	{
+		if (UDMXPixelMappingOutputComponent* ChildComponent = WeakOutputComponent.Get())
+		{
+			ChildComponent->PreEditChange(nullptr);
 
-						FVector2D AnchorOffset = Anchor - ChildComponent->GetPosition();
+			constexpr bool bModifyChildrenRecursive = true;
+			ChildComponent->ForEachChild([](UDMXPixelMappingBaseComponent* Component)
+				{
+					Component->Modify();
+				}, bModifyChildrenRecursive);
 
-						FVector2D NewPosition = GraphSpacePosition - AnchorOffset - PinnedDragDropOp->GraphSpaceDragOffset;
-						SetPosition(ChildComponent, NewPosition);
+			const FVector2D AnchorOffset = Anchor - ChildComponent->GetPosition();
 
-						ChildComponent->PostEditChange();
-					}
-				}
-			}
+			const FVector2D NewPosition = GraphSpacePosition - AnchorOffset - DragDropOp->GraphSpaceDragOffset;
+			ChildComponent->SetPosition(NewPosition);
 		}
 	}
-}
 
-void FDMXPixelMappingGroupChildDragDropHelper::SetPosition(UDMXPixelMappingOutputComponent* Component, const FVector2D& Position) const
-{
-	if (Component)
-	{
-		Component->Modify();
+	DragDropOp->GridSnap();
 
-		Component->SetPosition(Position);
+	for (const TWeakObjectPtr<UDMXPixelMappingOutputComponent>& WeakOutputComponent : WeakChildComponents)
+	{
+		if (UDMXPixelMappingOutputComponent* ChildComponent = WeakOutputComponent.Get())
+		{
+			ChildComponent->PostEditChange();
+		}
 	}
 }
diff --git a/Engine/Plugins/VirtualProduction/DMX/DMXPixelMapping/Source/DMXPixelMappingEditor/Private/DragDrop/DMXPixelMappingGroupChildDragDropHelper.h b/Engine/Plugins/VirtualProduction/DMX/DMXPixelMapping/Source/DMXPixelMappingEditor/Private/DragDrop/DMXPixelMappingGroupChildDragDropHelper.h
--- a/Engine/Plugins/VirtualProduction/DMX/DMXPixelMapping/Source/DMXPixelMappingEditor/Private/DragDrop/DMXPixelMappingGroupChildDragDropHelper.h
+++ b/Engine/Plugins/VirtualProduction/DMX/DMXPixelMapping/Source/DMXPixelMappingEditor/Private/DragDrop/DMXPixelMappingGroupChildDragDropHelper.h
@@ -35,9 +35,6 @@ private:
 	/** Helper that lays out cells unaligned */
 	void LayoutUnaligned(const FVector2D& GraphSpacePosition);
 
-	/** Sets the position of the component */
-	void SetPosition(UDMXPixelMappingOutputComponent* Component, const FVector2D& Position) const;
-
 	/** Local position of the parent */
 	FVector2D ParentPosition;
 
diff --git a/Engine/Plugins/VirtualProduction/DMX/DMXPixelMapping/Source/DMXPixelMappingEditor/Private/Toolkits/DMXPixelMappingToolkit.cpp b/Engine/Plugins/VirtualProduction/DMX/DMXPixelMapping/Source/DMXPixelMappingEditor/Private/Toolkits/DMXPixelMappingToolkit.cpp
--- a/Engine/Plugins/VirtualProduction/DMX/DMXPixelMapping/Source/DMXPixelMappingEditor/Private/Toolkits/DMXPixelMappingToolkit.cpp
+++ b/Engine/Plugins/VirtualProduction/DMX/DMXPixelMapping/Source/DMXPixelMappingEditor/Private/Toolkits/DMXPixelMappingToolkit.cpp
@@ -466,8 +466,8 @@ void FDMXPixelMappingToolkit::SizeSelectedComponentToTexture(bool bTransacted)
 				Child->PreEditChange(UDMXPixelMappingOutputComponent::StaticClass()->FindPropertyByName(UDMXPixelMappingOutputComponent::GetSizeXPropertyName()));
 				Child->PreEditChange(UDMXPixelMappingOutputComponent::StaticClass()->FindPropertyByName(UDMXPixelMappingOutputComponent::GetSizeYPropertyName()));
 
-				// Scale size (SetSize already clamps)
-				Child->SetSize(Child->GetSize() * RatioVector);
+				// Scale size
+				FVector2D NewSize = Child->GetSize() * RatioVector;
 
 				// Scale position (new position is zero vector)
 				const FVector2D ChildPosition = Child->GetPosition();
@@ -490,6 +490,41 @@ void FDMXPixelMappingToolkit::SizeSelectedComponentToTexture(bool bTransacted)
 	Component->PostEditChange();
 }
 
+void FDMXPixelMappingToolkit::ToggleGridSnapping()
+{
+	if (UDMXPixelMapping* PixelMapping = GetDMXPixelMapping())
+	{
+		const FScopedTransaction AddMappingTransaction(LOCTEXT("ToggleGridSnappingTransaction", "Toggle Grid Snapping"));
+		PixelMapping->PreEditChange(UDMXPixelMapping::StaticClass()->FindPropertyByName(GET_MEMBER_NAME_CHECKED(UDMXPixelMapping, bGridSnappingEnabled)));
+		PixelMapping->bGridSnappingEnabled = !PixelMapping->bGridSnappingEnabled;
+		PixelMapping->PostEditChange();
+	}
+}
+
+void FDMXPixelMappingToolkit::PostUndo(bool bSuccess)
+{
+	UDMXPixelMappingRootComponent* RootComponent = DMXPixelMapping ? DMXPixelMapping->GetRootComponent() : nullptr;
+	if (!RootComponent)
+	{
+		return;
+	}
+
+	constexpr bool bRecursive = false;
+	RootComponent->ForEachChild([](UDMXPixelMappingBaseComponent* Component)
+		{
+			if (UDMXPixelMappingRendererComponent* RendererComponent = Cast<UDMXPixelMappingRendererComponent>(Component))
+			{
+				RendererComponent->UpdatePreprocessRenderer();
+			}
+		}, bRecursive);
+}
+
+void FDMXPixelMappingToolkit::PostRedo(bool bSuccess)
+{
+	// Same behaviour as PostUndo
+	PostUndo(bSuccess);
+}
+
 void FDMXPixelMappingToolkit::OnComponentAddedOrRemoved(UDMXPixelMapping* PixelMapping, UDMXPixelMappingBaseComponent* Component)
 {
 	if (!bAddingComponents && !bRemovingComponents)
@@ -851,6 +886,11 @@ void FDMXPixelMappingToolkit::SetupCommands()
 			})
 	);
 
+	GetToolkitCommands()->MapAction(
+		FDMXPixelMappingEditorCommands::Get().ToggleGridSnapping,
+		FExecuteAction::CreateSP(this, &FDMXPixelMappingToolkit::ToggleGridSnapping)
+	);
+
 	// Designer related
 	constexpr bool bTransactSizeComponentToTexture = true;
 	GetToolkitCommands()->MapAction(
diff --git a/Engine/Plugins/VirtualProduction/DMX/DMXPixelMapping/Source/DMXPixelMappingEditor/Private/Toolkits/DMXPixelMappingToolkit.h b/Engine/Plugins/VirtualProduction/DMX/DMXPixelMapping/Source/DMXPixelMappingEditor/Private/Toolkits/DMXPixelMappingToolkit.h
--- a/Engine/Plugins/VirtualProduction/DMX/DMXPixelMapping/Source/DMXPixelMappingEditor/Private/Toolkits/DMXPixelMappingToolkit.h
+++ b/Engine/Plugins/VirtualProduction/DMX/DMXPixelMapping/Source/DMXPixelMappingEditor/Private/Toolkits/DMXPixelMappingToolkit.h
@@ -3,6 +3,7 @@
 #pragma once
 
 #include "DMXPixelMappingComponentReference.h"
+#include "EditorUndoClient.h"
 #include "Settings/DMXPixelMappingEditorSettings.h"
 #include "TickableEditorObject.h"
 #include "Toolkits/AssetEditorToolkit.h"
@@ -32,6 +33,7 @@ class UDMXPixelMappingRendererComponent;
 class FDMXPixelMappingToolkit
 	: public FAssetEditorToolkit
 	, public FTickableEditorObject
+	, public FSelfRegisteringEditorUndoClient
 {
 public:
 	DECLARE_MULTICAST_DELEGATE(FOnSelectedComponentsChangedDelegate)
@@ -136,7 +138,15 @@ public:
 	/** Sizes the component to the render target of the pixelmapping asset */
 	void SizeSelectedComponentToTexture(bool bTransacted);
 
+	/** Toggles between grid snapping enabled and disabled */
+	void ToggleGridSnapping();
+
 private:
+	//~ Begin FSelfRegisteringEditorUndoClient interface
+	virtual void PostUndo(bool bSuccess) override;
+	virtual void PostRedo(bool bSuccess) override;
+	//~ End FSelfRegisteringEditorUndoClient interface
+
 	/** Called when a component was added to the pixel mapping */
 	void OnComponentAddedOrRemoved(UDMXPixelMapping* PixelMapping, UDMXPixelMappingBaseComponent* Component);
 
diff --git a/Engine/Plugins/VirtualProduction/DMX/DMXPixelMapping/Source/DMXPixelMappingEditor/Private/ViewModels/DMXPixelMappingDMXLibraryViewModel.cpp b/Engine/Plugins/VirtualProduction/DMX/DMXPixelMapping/Source/DMXPixelMappingEditor/Private/ViewModels/DMXPixelMappingDMXLibraryViewModel.cpp
--- a/Engine/Plugins/VirtualProduction/DMX/DMXPixelMapping/Source/DMXPixelMappingEditor/Private/ViewModels/DMXPixelMappingDMXLibraryViewModel.cpp
+++ b/Engine/Plugins/VirtualProduction/DMX/DMXPixelMapping/Source/DMXPixelMappingEditor/Private/ViewModels/DMXPixelMappingDMXLibraryViewModel.cpp
@@ -2,6 +2,7 @@
 
 #include "DMXPixelMappingDMXLibraryViewModel.h"
 
+#include "Algo/MaxElement.h"
 #include "Algo/Transform.h"
 #include "Components/DMXPixelMappingBaseComponent.h"
 #include "Components/DMXPixelMappingFixtureGroupComponent.h"
@@ -10,14 +11,17 @@
 #include "Components/DMXPixelMappingRendererComponent.h"
 #include "Components/DMXPixelMappingRootComponent.h"
 #include "DMXPixelMapping.h"
+#include "DMXPixelMappingEditorLog.h"
 #include "Editor.h"
 #include "Library/DMXEntityFixturePatch.h"
 #include "Library/DMXLibrary.h"
-#include "DMXPixelMappingEditorLog.h"
+#include "ScopedTransaction.h"
 #include "Templates/DMXPixelMappingComponentTemplate.h"
 #include "Toolkits/DMXPixelMappingToolkit.h"
 
 
+#define LOCTEXT_NAMESPACE "DMXPixelMappingDMXLibraryViewModel"
+
 void UDMXPixelMappingDMXLibraryViewModel::CreateAndSetNewFixtureGroup(TWeakPtr<FDMXPixelMappingToolkit> InWeakToolkit)
 {
 	if (!InWeakToolkit.IsValid())
@@ -33,6 +37,10 @@ void UDMXPixelMappingDMXLibraryViewModel::CreateAndSetNewFixtureGroup(TWeakPtr<F
 	{
 		return;
 	}
+
+	const FScopedTransaction AddNewFixtureGroupTransaction(LOCTEXT("AddNewFixtureGroupTransaction", "Add Fixture Group to Pixel Mapping"));
+	RootComponent->Modify();
+
 	const TArray<UDMXPixelMappingFixtureGroupComponent*> OtherFixtureGroupComponents = GetFixtureGroupComponentsOfSameLibrary();
 
 	const TSharedRef<FDMXPixelMappingComponentTemplate> Template = MakeShared<FDMXPixelMappingComponentTemplate>(UDMXPixelMappingFixtureGroupComponent::StaticClass());
@@ -51,24 +59,22 @@ void UDMXPixelMappingDMXLibraryViewModel::CreateAndSetNewFixtureGroup(TWeakPtr<F
 		// If there's no group, add one that scales the texture of the active renderer component
 		if (Toolkit->CanSizeSelectedComponentToTexture())
 		{
-			constexpr bool bTransacted = true;
+			constexpr bool bTransacted = false;
 			Toolkit->SizeSelectedComponentToTexture(bTransacted);
 		}
 	}
 	else
 	{
 		// If there's already a group, offset over the top left of the existing group
-		FVector2D NewPosition(0.f, 0.f);
-		FVector2D NewSize = NewFixtureGroupComponent->GetSize();
-		for (UDMXPixelMappingFixtureGroupComponent* Other : OtherFixtureGroupComponents)
-		{
-			if (Other->GetPosition().Y > NewPosition.Y)
+		const UDMXPixelMappingFixtureGroupComponent* const* MostOffsetOtherPtr = Algo::MaxElementBy(OtherFixtureGroupComponents,
+			[](const UDMXPixelMappingFixtureGroupComponent* Other)
 			{
-				NewPosition = Other->GetPosition();
-				NewSize = Other->GetSize();
-			}
-		}
-		NewPosition += FVector2D(FMath::Max(1.f, NewSize.X / 16), FMath::Max(1.f, NewSize.Y / 16));
+				return Other->GetPosition().Length();
+			});
+		checkf(MostOffsetOtherPtr, TEXT("No result from array that was tested to not be empty."));
+
+		FVector2D NewSize = ActiveRendererComponent->GetSize();
+		FVector2D NewPosition = (*MostOffsetOtherPtr)->GetPosition() + FVector2D(FMath::Max(1.f, NewSize.X / 16), FMath::Max(1.f, NewSize.Y / 16));
 
 		NewFixtureGroupComponent->SetPosition(NewPosition);
 		NewFixtureGroupComponent->SetSize(NewSize);
@@ -428,3 +434,5 @@ TArray<UDMXPixelMappingFixtureGroupComponent*> UDMXPixelMappingDMXLibraryViewMod
 
 	return FixtureGroupComponents;
 }
+
+#undef LOCTEXT_NAMESPACE
diff --git a/Engine/Plugins/VirtualProduction/DMX/DMXPixelMapping/Source/DMXPixelMappingEditor/Private/Viewports/DMXPixelMappingSourceTextureViewportClient.cpp b/Engine/Plugins/VirtualProduction/DMX/DMXPixelMapping/Source/DMXPixelMappingEditor/Private/Viewports/DMXPixelMappingSourceTextureViewportClient.cpp
--- a/Engine/Plugins/VirtualProduction/DMX/DMXPixelMapping/Source/DMXPixelMappingEditor/Private/Viewports/DMXPixelMappingSourceTextureViewportClient.cpp
+++ b/Engine/Plugins/VirtualProduction/DMX/DMXPixelMapping/Source/DMXPixelMappingEditor/Private/Viewports/DMXPixelMappingSourceTextureViewportClient.cpp
@@ -5,6 +5,7 @@
 #include "CanvasItem.h"
 #include "CanvasTypes.h"
 #include "Components/DMXPixelMappingRendererComponent.h"
+#include "DMXPixelMapping.h"
 #include "Engine/Texture.h"
 #include "Framework/Application/SlateApplication.h"
 #include "Slate/SceneViewport.h"
@@ -81,8 +82,14 @@ void FDMXPixelMappingSourceTextureViewportClient::Draw(FViewport* InViewport, FC
 		return;
 	}
 	const TSharedRef<FDMXPixelMappingToolkit> Toolkit = WeakToolkit.Pin().ToSharedRef();
-	const TSharedRef<SDMXPixelMappingDesignerView> DesignerView = Toolkit->GetOrCreateDesignerView();
 
+	UDMXPixelMapping* PixelMapping = Toolkit->GetDMXPixelMapping();
+	if (!PixelMapping)
+	{
+		return;
+	}
+
+	const TSharedRef<SDMXPixelMappingDesignerView> DesignerView = Toolkit->GetOrCreateDesignerView();
 	UDMXPixelMappingRendererComponent* RendererComponent = Toolkit->GetActiveRendererComponent();
 	UTexture* InputTexture = RendererComponent ? RendererComponent->GetRenderedInputTexture() : nullptr;
 	if (!InputTexture || !InputTexture->GetResource())
@@ -101,9 +108,9 @@ void FDMXPixelMappingSourceTextureViewportClient::Draw(FViewport* InViewport, FC
 	const FBox2D TextureBox(FVector2D::ZeroVector, TextureDimensions);
 
 	// Get the visible rect in grap space, and the UVs
-	const FBox2D VisibleRectBox = GetVisibleTextureBoxGraphSpace();
-	const FVector2D UV0 = VisibleRectBox.Min / TextureBox.GetSize();
-	const FVector2D UV1 = VisibleRectBox.Max / TextureBox.GetSize();
+	const FBox2D VisibleRect = GetVisibleTextureBoxGraphSpace();
+	const FVector2D UV0 = VisibleRect.Min / TextureBox.GetSize();
+	const FVector2D UV1 = VisibleRect.Max / TextureBox.GetSize();
 
 	if (IsDrawingVisibleRectOnly())
 	{
@@ -119,7 +126,7 @@ void FDMXPixelMappingSourceTextureViewportClient::Draw(FViewport* InViewport, FC
 	}
 	else
 	{
-		FCanvasTileItem TileItem(VisibleRectBox.Min * DesignerView->GetZoomAmount(), InputTexture->GetResource(), VisibleRectBox.GetSize() * DesignerView->GetZoomAmount(), UV0, UV1, FLinearColor::White);
+		FCanvasTileItem TileItem(VisibleRect.Min * DesignerView->GetZoomAmount(), InputTexture->GetResource(), VisibleRect.GetSize() * DesignerView->GetZoomAmount(), UV0, UV1, FLinearColor::White);
 		TileItem.BlendMode = ESimpleElementBlendMode::SE_BLEND_MAX;
 		Canvas->DrawItem(TileItem);
 	}
diff --git a/Engine/Plugins/VirtualProduction/DMX/DMXPixelMapping/Source/DMXPixelMappingEditor/Private/Views/SDMXPixelMappingDesignerView.cpp b/Engine/Plugins/VirtualProduction/DMX/DMXPixelMapping/Source/DMXPixelMappingEditor/Private/Views/SDMXPixelMappingDesignerView.cpp
--- a/Engine/Plugins/VirtualProduction/DMX/DMXPixelMapping/Source/DMXPixelMappingEditor/Private/Views/SDMXPixelMappingDesignerView.cpp
+++ b/Engine/Plugins/VirtualProduction/DMX/DMXPixelMapping/Source/DMXPixelMappingEditor/Private/Views/SDMXPixelMappingDesignerView.cpp
@@ -13,7 +13,6 @@
 #include "Components/DMXPixelMappingMatrixCellComponent.h"
 #include "Components/DMXPixelMappingScreenComponent.h"
 #include "DragDrop/DMXPixelMappingDragDropOp.h"
-#include "Input/HittestGrid.h"
 #include "ScopedTransaction.h"
 #include "Settings/DMXPixelMappingEditorSettings.h"
 #include "Toolkits/DMXPixelMappingToolkit.h"
@@ -23,6 +22,7 @@
 #include "Widgets/Layout/SGridPanel.h"
 #include "Widgets/SCanvas.h"
 #include "Widgets/SDMXPixelMappingDesignerCanvas.h"
+#include "Widgets/SDMXPixelMappingDesignerToolbar.h"
 #include "Widgets/SDMXPixelMappingOutputComponent.h"
 #include "Widgets/SDMXPixelMappingRuler.h"
 #include "Widgets/SDMXPixelMappingSourceTextureViewport.h"
@@ -32,8 +32,10 @@
 
 #define LOCTEXT_NAMESPACE "SDMXPixelMappingDesignerView"
 
-void SDMXPixelMappingDesignerView::Construct(const FArguments& InArgs, const TSharedPtr<FDMXPixelMappingToolkit>& InToolkit)
+void SDMXPixelMappingDesignerView::Construct(const FArguments& InArgs, const TSharedRef<FDMXPixelMappingToolkit>& InToolkit)
 {
+	using namespace UE::DMX;
+
 	SDMXPixelMappingSurface::Construct(SDMXPixelMappingSurface::FArguments()
 		.AllowContinousZoomInterpolation(false)
 		.Content()
@@ -77,7 +79,8 @@ void SDMXPixelMappingDesignerView::Construct(const FArguments& InArgs, const TSh
 					]
 				]
 			]
-			+SVerticalBox::Slot()
+
+			+ SVerticalBox::Slot()
 			[
 				SNew(SGridPanel)
 				.FillColumn(1, 1.0f)
@@ -125,23 +128,20 @@ void SDMXPixelMappingDesignerView::Construct(const FArguments& InArgs, const TSh
 						[
 							SNew(SOverlay)
 
+							// Viewport
 							+ SOverlay::Slot()
 							[
 								SAssignNew(SourceTextureViewport, SDMXPixelMappingSourceTextureViewport, InToolkit)
 							]
 
-							+ SOverlay::Slot()
-							[
-								SAssignNew(PreviewSizeConstraint, SBox)
-							]
-
+							// Designer canvas
 							+ SOverlay::Slot()
 							[
 								SAssignNew(DesignCanvasBorder, SBorder)
 								.HAlign(HAlign_Fill)
 								.VAlign(VAlign_Fill)
 								.BorderImage(FAppStyle::GetBrush("NoBorder"))
-								.Padding(0.0f)
+								.Padding(0.f)
 								[
 									SAssignNew(DesignCanvas, SConstraintCanvas)
 									.Visibility(EVisibility::HitTestInvisible)
@@ -159,63 +159,55 @@ void SDMXPixelMappingDesignerView::Construct(const FArguments& InArgs, const TSh
 						.Visibility(this, &SDMXPixelMappingDesignerView::GetExtensionCanvasVisibility)
 					]
 
-					// Designer overlay UI, toolbar, status messages, zoom level...etc
+					// Top toolbar UI
 					+ SOverlay::Slot()
 					.HAlign(HAlign_Fill)
-					.VAlign(VAlign_Fill)
+					.VAlign(VAlign_Top)
 					[
-						 SNew(SOverlay)
-						+ SOverlay::Slot()
-						.Padding(0)
-						+SOverlay::Slot()
-						.HAlign(HAlign_Fill)
-						.VAlign(VAlign_Top)
+						SNew(SHorizontalBox)
+
+						// Zoom text
+						+ SHorizontalBox::Slot()
+						.AutoWidth()
+						.VAlign(VAlign_Center)
+						.Padding(4.f, 2.0f)
 						[
-							SNew(SHorizontalBox)
+							SNew(STextBlock)
+							.TextStyle(FAppStyle::Get(), "Graph.ZoomText")
+							.Text(this, &SDMXPixelMappingDesignerView::GetZoomText)
+							.ColorAndOpacity(this, &SDMXPixelMappingDesignerView::GetZoomTextColorAndOpacity)
+						]
 
-							+ SHorizontalBox::Slot()
-							.AutoWidth()
-							.VAlign(VAlign_Center)
-							.Padding(6, 2, 0, 0)
-							[
-								SNew(STextBlock)
-								.TextStyle(FAppStyle::Get(), "Graph.ZoomText")
-								.Text(this, &SDMXPixelMappingDesignerView::GetZoomText)
-								.ColorAndOpacity(this, &SDMXPixelMappingDesignerView::GetZoomTextColorAndOpacity)
-							]
-							+ SHorizontalBox::Slot()
-							.AutoWidth()
-							.VAlign(VAlign_Center)
-							.Padding(40, 2, 0, 0)
-							[
-								SNew(STextBlock)
-								.TextStyle(FAppStyle::Get(), "Graph.ZoomText")
-								.Font(FCoreStyle::GetDefaultFontStyle(TEXT("BoldCondensed"), 14))
-								.Text(this, &SDMXPixelMappingDesignerView::GetCursorPositionText)
-								.ColorAndOpacity(FLinearColor(1.f, 1.f, 1.f, 0.25f))
-								.Visibility(this, &SDMXPixelMappingDesignerView::GetCursorPositionTextVisibility)
-							]
-							+ SHorizontalBox::Slot()
-							.FillWidth(1.0f)
-							[
-								SNew(SSpacer)
-								.Size(FVector2D(1, 1))
-							]
+						// Cursor position
+						+ SHorizontalBox::Slot()
+						.AutoWidth()
+						.VAlign(VAlign_Center)
+						.Padding(4.f, 2.0f)
+						[
+							SNew(STextBlock)
+							.TextStyle(FAppStyle::Get(), "Graph.ZoomText")
+							.Font(FCoreStyle::GetDefaultFontStyle(TEXT("BoldCondensed"), 14))
+							.Text(this, &SDMXPixelMappingDesignerView::GetCursorPositionText)
+							.ColorAndOpacity(FLinearColor(1.f, 1.f, 1.f, 0.25f))
+							.Visibility(this, &SDMXPixelMappingDesignerView::GetCursorPositionTextVisibility)
+						]
 
-							+ SHorizontalBox::Slot()
-							.AutoWidth()
-							.VAlign(VAlign_Center)
-							[
-								SNew(SButton)
-								.ButtonColorAndOpacity(FLinearColor::Transparent)
-								.ToolTipText(LOCTEXT("ZoomToFit_ToolTip", "Zoom To Fit"))
-								.OnClicked(this, &SDMXPixelMappingDesignerView::OnZoomToFitClicked)
-								.ContentPadding(FAppStyle::Get().GetMargin("ViewportMenu.SToolBarButtonBlock.Button.Padding"))
-								[
-									SNew(SImage)
-									.Image(FDMXPixelMappingEditorStyle::Get().GetBrush("Icons.ZoomToFit"))
-								]
-							]
+						// Spacer
+						+ SHorizontalBox::Slot()
+						.FillWidth(1.0f)
+						[
+							SNew(SSpacer)
+							.Size(FVector2D(1, 1))
+						]
+
+						// Toolbar
+						+ SHorizontalBox::Slot()
+						.AutoWidth()
+						.VAlign(VAlign_Center)
+						.Padding(4.f, 2.0f)
+						[
+							SNew(SDMXPixelMappingDesignerToolbar, InToolkit)
+							.OnZoomToFitClicked(this, &SDMXPixelMappingDesignerView::OnZoomToFitClicked)
 						]
 					]
 				]
@@ -225,8 +217,6 @@ void SDMXPixelMappingDesignerView::Construct(const FArguments& InArgs, const TSh
 
 	ZoomToFit(true);
 
-	HittestGrid = MakeShared<FHittestGrid>();
-
 	// Bind to selection changes
 	InToolkit->GetOnSelectedComponentsChangedDelegate().AddSP(this, &SDMXPixelMappingDesignerView::OnSelectedComponentsChanged);
 
@@ -240,6 +230,143 @@ const FGeometry& SDMXPixelMappingDesignerView::GetGraphTickSpaceGeometry() const
 	return ZoomPan->GetTickSpaceGeometry();
 }
 
+int32 SDMXPixelMappingDesignerView::OnPaint(const FPaintArgs& Args, const FGeometry& AllottedGeometry, const FSlateRect& MyCullingRect, FSlateWindowElementList& OutDrawElements, int32 LayerId, const FWidgetStyle& InWidgetStyle, bool bParentEnabled) const
+{
+	LayerId = SDMXPixelMappingSurface::OnPaint(Args, AllottedGeometry, MyCullingRect, OutDrawElements, LayerId, InWidgetStyle, bParentEnabled);
+
+	// Draw the grid snapping grid overlay if required
+	const UDMXPixelMapping* PixelMapping = WeakToolkit.IsValid() ? WeakToolkit.Pin()->GetDMXPixelMapping() : nullptr;
+	if (!PixelMapping || 
+		!PixelMapping->bGridSnappingEnabled || 
+		PixelMapping->SnapGridColor.A == 0.f || 
+		!CachedRendererComponent.IsValid())
+	{
+		return LayerId;
+	}
+	
+	if (!ensureMsgf(PixelMapping->SnapGridColumns > 0 && PixelMapping->SnapGridRows, TEXT("Invalid snap grid. Cannot draw grid.")))
+	{
+		return LayerId;
+	}
+
+	// Get the line spacings
+	const FVector2D RelativeOrigin = AllottedGeometry.AbsoluteToLocal(GridOrigin);
+	const FVector2D TextureSizeZoomed = CachedRendererComponent->GetSize() * GetZoomAmount();
+	const float LineSpacingX = TextureSizeZoomed.X / PixelMapping->SnapGridColumns;
+	const float LineSpacingY = TextureSizeZoomed.Y / PixelMapping->SnapGridRows;
+
+	// Draw the grid above the surface, above the hit test root overlay, above the the zoom pan overlay, above the texture overlay
+	const int32 GridLayer = LayerId + 4;
+
+	// Horizontal
+	TArray<FVector2D> LinePoints;
+	LinePoints.SetNumUninitialized(2);
+	for (int32 RowIndex = 0; RowIndex <= PixelMapping->SnapGridRows; RowIndex++)
+	{
+		LinePoints[0] = RelativeOrigin + FVector2D(0.f, LineSpacingY * RowIndex);
+		LinePoints[1] = RelativeOrigin + FVector2D(TextureSizeZoomed.X, LineSpacingY * RowIndex);
+
+		constexpr bool bAntialias = false;
+		FSlateDrawElement::MakeLines(
+			OutDrawElements,
+			GridLayer,
+			AllottedGeometry.ToPaintGeometry(),
+			LinePoints,
+			ESlateDrawEffect::None,
+			PixelMapping->SnapGridColor,
+			bAntialias);
+	}
+
+	// Vertical
+	for (int32 ColumnIndex = 0; ColumnIndex <= PixelMapping->SnapGridColumns; ColumnIndex++)
+	{
+		LinePoints[0] = RelativeOrigin + FVector2D(LineSpacingX * ColumnIndex, 0.f);
+		LinePoints[1] = RelativeOrigin + FVector2D(LineSpacingX * ColumnIndex, TextureSizeZoomed.Y);
+
+		constexpr bool bAntialias = false;
+		FSlateDrawElement::MakeLines(
+			OutDrawElements, 
+			GridLayer,
+			AllottedGeometry.ToPaintGeometry(), 
+			LinePoints, 
+			ESlateDrawEffect::None, 
+			PixelMapping->SnapGridColor,
+			bAntialias);
+	}
+
+	return GridLayer;
+}
+
+void SDMXPixelMappingDesignerView::Tick(const FGeometry& AllottedGeometry, const double InCurrentTime, const float InDeltaTime)
+{
+	SDMXPixelMappingSurface::Tick(AllottedGeometry, InCurrentTime, InDeltaTime);
+
+	const TSharedPtr<FDMXPixelMappingToolkit> Toolkit = WeakToolkit.Pin();
+	if (!Toolkit.IsValid())
+	{
+		return;
+	}
+
+	// Update test if the renderer changed, and if so, rebuild the designer
+	UDMXPixelMappingRendererComponent* RendererComponent = Toolkit->GetActiveRendererComponent();
+	if (RendererComponent != CachedRendererComponent)
+	{
+		CachedRendererComponent = RendererComponent;
+		RebuildDesigner();
+	}
+
+	// Cache geometry
+	CachedWidgetGeometry.Reset();
+	FArrangedWidget WindowWidgetGeometry(PreviewHitTestRoot.ToSharedRef(), GetDesignerGeometry());
+	PopulateWidgetGeometryCache(WindowWidgetGeometry);
+
+	// Compute the origin in absolute screen space.
+	FGeometry RootGeometry = CachedWidgetGeometry.FindChecked(SourceTextureViewport.ToSharedRef()).Geometry;
+	GridOrigin = MakeGeometryWindowLocal(RootGeometry).LocalToAbsolute(FVector2D::ZeroVector);
+
+	// Update the rulers
+	TopRuler->SetRuling(GridOrigin, 1.0f / GetZoomAmount());
+	SideRuler->SetRuling(GridOrigin, 1.0f / GetZoomAmount());
+
+	if (IsHovered())
+	{
+		// Get cursor in absolute window space.
+		FVector2D CursorPos = FSlateApplication::Get().GetCursorPos();
+		CursorPos = MakeGeometryWindowLocal(RootGeometry).LocalToAbsolute(RootGeometry.AbsoluteToLocal(CursorPos));
+
+		TopRuler->SetCursor(CursorPos);
+		SideRuler->SetCursor(CursorPos);
+	}
+	else
+	{
+		TopRuler->SetCursor(TOptional<FVector2D>());
+		SideRuler->SetCursor(TOptional<FVector2D>());
+	}
+
+	// Handle pending drag drop ops
+	if (PendingDragDropOp.IsValid())
+	{
+		FVector2D GraphSpaceCursorPosition;
+		if (GetGraphSpaceCursorPosition(GraphSpaceCursorPosition))
+		{
+			// Layout childs with the group item drag drop helper if possible, else use the default layout method
+			if (TSharedPtr<FDMXPixelMappingGroupChildDragDropHelper> GroupChildDragDropHelper = PendingDragDropOp->GetGroupChildDragDropHelper())
+			{
+				// If it was created from template, it was dragged in from details or palette. In this case align the components.
+				const bool bAlignComponents = PendingDragDropOp->WasCreatedAsTemplate();
+
+				GroupChildDragDropHelper->LayoutComponents(GraphSpaceCursorPosition, bAlignComponents);
+			}
+			else
+			{
+				PendingDragDropOp->LayoutOutputComponents(GraphSpaceCursorPosition);
+			}
+		}
+
+		PendingDragDropOp.Reset();
+	}
+}
+
 FReply SDMXPixelMappingDesignerView::OnMouseButtonDown(const FGeometry& MyGeometry, const FPointerEvent& MouseEvent)
 {
 	SDMXPixelMappingSurface::OnMouseButtonDown(MyGeometry, MouseEvent);
@@ -332,42 +459,6 @@ FReply SDMXPixelMappingDesignerView::OnMouseButtonUp(const FGeometry& MyGeometry
 		.SetUserFocus(AsShared());
 }
 
-FReply SDMXPixelMappingDesignerView::OnMouseMove(const FGeometry& MyGeometry, const FPointerEvent& MouseEvent)
-{
-	if (MouseEvent.GetCursorDelta().IsZero())
-	{
-		return FReply::Unhandled();
-	}
-
-	FReply SurfaceHandled = SDMXPixelMappingSurface::OnMouseMove(MyGeometry, MouseEvent);
-	if (SurfaceHandled.IsEventHandled())
-	{
-		return SurfaceHandled;
-	}
-
-	if (MouseEvent.IsMouseButtonDown(EKeys::LeftMouseButton) && HasMouseCapture())
-	{
-		bool bIsRootWidgetSelected = false;
-		for (const FDMXPixelMappingComponentReference& ComponentReference : GetSelectedComponents())
-		{
-			UDMXPixelMappingBaseComponent* Component = ComponentReference.GetComponent();
-			if (Component && Component->GetParent())
-			{
-				bIsRootWidgetSelected = true;
-				break;
-			}
-		}
-
-		if (!bIsRootWidgetSelected)
-		{
-			//Drag selected widgets
-			return FReply::Handled();
-		}
-	}
-
-	return FReply::Unhandled();
-}
-
 FReply SDMXPixelMappingDesignerView::OnKeyDown(const FGeometry& MyGeometry, const FKeyEvent& InKeyEvent)
 {
 	if (InKeyEvent.GetKey() == EKeys::Delete)
@@ -385,133 +476,65 @@ FReply SDMXPixelMappingDesignerView::OnKeyDown(const FGeometry& MyGeometry, cons
 	return FReply::Unhandled();
 }
 
-void SDMXPixelMappingDesignerView::Tick(const FGeometry& AllottedGeometry, const double InCurrentTime, const float InDeltaTime)
-{
-	SDMXPixelMappingSurface::Tick(AllottedGeometry, InCurrentTime, InDeltaTime);
-
-	if (const TSharedPtr<FDMXPixelMappingToolkit> Toolkit = WeakToolkit.Pin())
-	{
-		// Handle drag over on tick
-		if (PendingDragDropOp.IsValid())
-		{
-			FVector2D GraphSpaceCursorPosition;
-			if (GetGraphSpaceCursorPosition(GraphSpaceCursorPosition))
-			{
-				// Layout childs with the group item drag drop helper if possible, else use the default layout method
-				if (TSharedPtr<FDMXPixelMappingGroupChildDragDropHelper> GroupChildDragDropHelper = PendingDragDropOp->GetGroupChildDragDropHelper())
-				{
-					// If it was created from template it was dragged in from details or palette. In this case align the components.
-					const bool bAlignComponents = PendingDragDropOp->WasCreatedAsTemplate();
-
-					GroupChildDragDropHelper->LayoutComponents(GraphSpaceCursorPosition, bAlignComponents);
-				}
-				else
-				{
-					PendingDragDropOp->LayoutOutputComponents(GraphSpaceCursorPosition);
-				}
-			}
-
-			PendingDragDropOp.Reset();
-		}
-
-		// Update test if the renderer changed, and if so, rebuild the designer
-		UDMXPixelMappingRendererComponent* RendererComponent = Toolkit->GetActiveRendererComponent();
-		if (RendererComponent != CachedRendererComponent)
-		{
-			CachedRendererComponent = RendererComponent;
-			RebuildDesigner();
-		}
-
-		// Cache geometry
-		CachedWidgetGeometry.Reset();
-		FArrangedWidget WindowWidgetGeometry(PreviewHitTestRoot.ToSharedRef(), GetDesignerGeometry());
-		PopulateWidgetGeometryCache(WindowWidgetGeometry);
-
-		// Compute the origin in absolute space.
-		FGeometry RootGeometry = CachedWidgetGeometry.FindChecked(PreviewSizeConstraint.ToSharedRef()).Geometry;
-		FVector2D AbsoluteOrigin = MakeGeometryWindowLocal(RootGeometry).LocalToAbsolute(FVector2D::ZeroVector);
-
-		GridOrigin = AbsoluteOrigin;
-
-		// Roler position
-		TopRuler->SetRuling(AbsoluteOrigin, 1.0f / GetZoomAmount());
-		SideRuler->SetRuling(AbsoluteOrigin, 1.0f / GetZoomAmount());
-
-		if (IsHovered())
-		{
-			// Get cursor in absolute window space.
-			FVector2D CursorPos = FSlateApplication::Get().GetCursorPos();
-			CursorPos = MakeGeometryWindowLocal(RootGeometry).LocalToAbsolute(RootGeometry.AbsoluteToLocal(CursorPos));
-
-			TopRuler->SetCursor(CursorPos);
-			SideRuler->SetCursor(CursorPos);
-		}
-		else
-		{
-			TopRuler->SetCursor(TOptional<FVector2D>());
-			SideRuler->SetCursor(TOptional<FVector2D>());
-		}
-	}
-}
-
 FReply SDMXPixelMappingDesignerView::OnDragDetected(const FGeometry& MyGeometry, const FPointerEvent& MouseEvent)
 {
 	SDMXPixelMappingSurface::OnDragDetected(MyGeometry, MouseEvent);
 
 	const TSet<FDMXPixelMappingComponentReference>& SelectedComponents = GetSelectedComponents();
-
-	if (SelectedComponents.Num() > 0)
+	if (!WeakToolkit.IsValid() || SelectedComponents.IsEmpty())
 	{
-		TArray<UDMXPixelMappingOutputComponent*> DraggedComponentCandidates;
+		return FReply::Handled();
+	}
+	const TSharedRef<FDMXPixelMappingToolkit> Toolkit = WeakToolkit.Pin().ToSharedRef();
 
-		for (const FDMXPixelMappingComponentReference& SelectedComponent : SelectedComponents)
+	TArray<UDMXPixelMappingOutputComponent*> DraggedComponentCandidates;
+	for (const FDMXPixelMappingComponentReference& SelectedComponent : SelectedComponents)
+	{
+		if (UDMXPixelMappingOutputComponent* OutputComponent = Cast<UDMXPixelMappingOutputComponent>(SelectedComponent.GetComponent()))
 		{
-			if (UDMXPixelMappingOutputComponent* OutputComponent = Cast<UDMXPixelMappingOutputComponent>(SelectedComponent.GetComponent()))
-			{
-				DraggedComponentCandidates.Add(OutputComponent);
-			}
+			DraggedComponentCandidates.Add(OutputComponent);
 		}
+	}
 
-		TArray<TWeakObjectPtr<UDMXPixelMappingBaseComponent>> DraggedComponents;
-		for (UDMXPixelMappingOutputComponent* Candidate : DraggedComponentCandidates)
+	TArray<TWeakObjectPtr<UDMXPixelMappingBaseComponent>> DraggedComponents;
+	for (UDMXPixelMappingOutputComponent* Candidate : DraggedComponentCandidates)
+	{
+		// Check the parent chain of each dragged component and ignore those that are children of other dragged components
+		bool bIsChild = false;
+		for (TWeakObjectPtr<UDMXPixelMappingBaseComponent> Parent = Candidate->GetParent(); Parent.IsValid(); Parent = Parent->GetParent())
 		{
-			// Check the parent chain of each dragged component and ignore those that are children of other dragged components
-			bool bIsChild = false;
-			for (TWeakObjectPtr<UDMXPixelMappingBaseComponent> Parent = Candidate->GetParent(); Parent.IsValid(); Parent = Parent->GetParent())
-			{
-				bIsChild = DraggedComponentCandidates.ContainsByPredicate([Parent](const UDMXPixelMappingOutputComponent* BaseComponent)
-					{
-						return BaseComponent == Parent.Get();
-					});
-
-				if (bIsChild)
+			bIsChild = DraggedComponentCandidates.ContainsByPredicate([Parent](const UDMXPixelMappingOutputComponent* BaseComponent)
 				{
-					break;
-				}
-			}
+					return BaseComponent == Parent.Get();
+				});
 
-			if (!bIsChild && !Candidate->IsLockInDesigner())
+			if (bIsChild)
 			{
-				DraggedComponents.Add(Candidate);
+				break;
 			}
 		}
 
-		// We know the first element is the one that was selected via mouse
-		if (DraggedComponents.Num() > 0)
+		if (!bIsChild && !Candidate->IsLockInDesigner())
 		{
-			// We rely on the first being the clicked component. There's an according comment in the DragDropOp and the detected drag is raised here.
-			UDMXPixelMappingOutputComponent* ClickedComponent = CastChecked<UDMXPixelMappingOutputComponent>(DraggedComponents[0]);
+			DraggedComponents.Add(Candidate);
+		}
+	}
 
-			const FVector2D GraphSpaceDragOffset = DragAnchor - ClickedComponent->GetPosition();
+	// We know the first element is the one that was selected via mouse
+	if (DraggedComponents.Num() > 0)
+	{
+		// We rely on the first being the clicked component. There's an according comment in the DragDropOp and the detected drag is raised here.
+		UDMXPixelMappingOutputComponent* ClickedComponent = CastChecked<UDMXPixelMappingOutputComponent>(DraggedComponents[0]);
 
-			TSharedRef<FDMXPixelMappingDragDropOp> DragDropOp = FDMXPixelMappingDragDropOp::New(GraphSpaceDragOffset, DraggedComponents);
-			DragDropOp->SetDecoratorVisibility(false);
+		const FVector2D GraphSpaceDragOffset = DragAnchor - ClickedComponent->GetPosition();
 
-			// Clear any pending selected widgets
-			PendingSelectedComponent = nullptr;
+		TSharedRef<FDMXPixelMappingDragDropOp> DragDropOp = FDMXPixelMappingDragDropOp::New(Toolkit, GraphSpaceDragOffset, DraggedComponents);
+		DragDropOp->SetDecoratorVisibility(false);
 
-			return FReply::Handled().BeginDragDrop(DragDropOp);
-		}
+		// Clear any pending selected widgets
+		PendingSelectedComponent = nullptr;
+
+		return FReply::Handled().BeginDragDrop(DragDropOp);
 	}
 
 	return FReply::Handled();
@@ -522,7 +545,7 @@ void SDMXPixelMappingDesignerView::OnDragEnter(const FGeometry& MyGeometry, cons
 	const TSharedPtr<FDMXPixelMappingDragDropOp> DragDropOp = DragDropEvent.GetOperationAs<FDMXPixelMappingDragDropOp>();
 	if (DragDropOp.IsValid())
 	{
-		const FGeometry WidgetUnderCursorGeometry = PreviewSizeConstraint->GetTickSpaceGeometry();
+		const FGeometry WidgetUnderCursorGeometry = SourceTextureViewport->GetTickSpaceGeometry();
 		const FVector2D ScreenSpacePosition = DragDropEvent.GetScreenSpacePosition();
 		const FVector2D LocalPosition = WidgetUnderCursorGeometry.AbsoluteToLocal(ScreenSpacePosition);
 
@@ -609,17 +632,17 @@ FSlateRect SDMXPixelMappingDesignerView::ComputeAreaBounds() const
 
 int32 SDMXPixelMappingDesignerView::GetGraphRulePeriod() const
 {
-	return 10; // Parent override
+	return 10;
 }
 
 float SDMXPixelMappingDesignerView::GetGridScaleAmount() const
 {
-	return 1.0f; // Parent override
+	return 1.f;
 }
 
-int32 SDMXPixelMappingDesignerView::GetSnapGridSize() const
+int32 SDMXPixelMappingDesignerView::GetGridSize() const
 {
-	return 4; // Parent override
+	return 4;
 }
 
 void SDMXPixelMappingDesignerView::PostUndo(bool bSuccess)
@@ -911,21 +934,12 @@ void SDMXPixelMappingDesignerView::PopulateWidgetGeometryCache(FArrangedWidget&
 {
 	const FSlateRect Rect = PreviewHitTestRoot->GetTickSpaceGeometry().GetLayoutBoundingRect();
 	const FSlateRect PaintRect = PreviewHitTestRoot->GetPaintSpaceGeometry().GetLayoutBoundingRect();
-	HittestGrid->SetHittestArea(Rect.GetTopLeft(), Rect.GetSize(), PaintRect.GetTopLeft());
-	HittestGrid->Clear();
 
 	PopulateWidgetGeometryCache_Loop(Root);
 }
 
 void SDMXPixelMappingDesignerView::PopulateWidgetGeometryCache_Loop(FArrangedWidget& CurrentWidget)
 {
-	bool bIncludeInHitTestGrid = true;
-
-	if (bIncludeInHitTestGrid)
-	{
-		HittestGrid->AddWidget(&(CurrentWidget.Widget.Get()), 0, 0, FSlateInvalidationWidgetSortOrder());
-	}
-
 	FArrangedChildren ArrangedChildren(EVisibility::All);
 	CurrentWidget.Widget->ArrangeChildren(CurrentWidget.Geometry, ArrangedChildren);
 
@@ -1028,10 +1042,9 @@ bool SDMXPixelMappingDesignerView::GetComponentGeometry(UDMXPixelMappingBaseComp
 
 bool SDMXPixelMappingDesignerView::GetGraphSpaceCursorPosition(FVector2D& OutGraphSpaceCursorPosition) const
 {
-	if (const FArrangedWidget* CachedPreviewSurface = CachedWidgetGeometry.Find(PreviewSizeConstraint.ToSharedRef()))
+	if (const FArrangedWidget* PreviewSurface = CachedWidgetGeometry.Find(SourceTextureViewport.ToSharedRef()))
 	{
-		const FGeometry& RootGeometry = CachedPreviewSurface->Geometry;
-		OutGraphSpaceCursorPosition = RootGeometry.AbsoluteToLocal(FSlateApplication::Get().GetCursorPos());
+		OutGraphSpaceCursorPosition = PreviewSurface->Geometry.AbsoluteToLocal(FSlateApplication::Get().GetCursorPos());
 
 		return true;
 	}
diff --git a/Engine/Plugins/VirtualProduction/DMX/DMXPixelMapping/Source/DMXPixelMappingEditor/Private/Views/SDMXPixelMappingDesignerView.h b/Engine/Plugins/VirtualProduction/DMX/DMXPixelMapping/Source/DMXPixelMappingEditor/Private/Views/SDMXPixelMappingDesignerView.h
--- a/Engine/Plugins/VirtualProduction/DMX/DMXPixelMapping/Source/DMXPixelMappingEditor/Private/Views/SDMXPixelMappingDesignerView.h
+++ b/Engine/Plugins/VirtualProduction/DMX/DMXPixelMapping/Source/DMXPixelMappingEditor/Private/Views/SDMXPixelMappingDesignerView.h
@@ -2,33 +2,27 @@
 
 #pragma once
 
-#include "Widgets/SDMXPixelMappingSurface.h"
-
-
 #include "EditorUndoClient.h"
-
-class UDMXPixelMapping;
+#include "Widgets/SDMXPixelMappingSurface.h"
 
 class FDMXPixelMappingComponentReference;
 class FDMXPixelMappingDragDropOp;
 class FDMXPixelMappingToolkit;
 class IDMXPixelMappingOutputComponentWidgetInterface;
+class SBorder;
+class SBox;
+class SCanvas;
+class SConstraintCanvas;
 class SDMXPixelMappingRuler;
 class SDMXPixelMappingSourceTextureViewport;
 class SDMXPixelMappingTransformHandle;
 class SDMXPixelMappingZoomPan;
+class SOverlay;
+class UDMXPixelMapping;
 class UDMXPixelMappingBaseComponent;
 class UDMXPixelMappingOutputComponent;
 class UDMXPixelMappingRendererComponent;
 
-class FHittestGrid;
-struct FOptionalSize;
-class SBorder;
-class SBox;
-class SCanvas;
-class SConstraintCanvas;
-class SOverlay;
-
 
 class SDMXPixelMappingDesignerView
 	: public SDMXPixelMappingSurface
@@ -58,36 +52,36 @@ public:
 	 *
 	 * @param InArgs The construction arguments.
 	 */
-	void Construct(const FArguments& InArgs, const TSharedPtr<FDMXPixelMappingToolkit>& InToolkit);
+	void Construct(const FArguments& InArgs, const TSharedRef<FDMXPixelMappingToolkit>& InToolkit);
 
 	/** Returns the geometry of the graph */
 	const FGeometry& GetGraphTickSpaceGeometry() const;
 
 protected:
-	// Begin SWidget interface
+	//~ Begin SWidget interface
+	virtual int32 OnPaint(const FPaintArgs& Args, const FGeometry& AllottedGeometry, const FSlateRect& MyCullingRect, FSlateWindowElementList& OutDrawElements, int32 LayerId, const FWidgetStyle& InWidgetStyle, bool bParentEnabled) const override;
+	virtual void Tick(const FGeometry& AllottedGeometry, const double InCurrentTime, const float InDeltaTime) override;
 	virtual FReply OnMouseButtonDown(const FGeometry& MyGeometry, const FPointerEvent& MouseEvent) override;
 	virtual FReply OnMouseButtonUp(const FGeometry& MyGeometry, const FPointerEvent& MouseEvent) override;
-	virtual FReply OnMouseMove(const FGeometry& MyGeometry, const FPointerEvent& MouseEvent) override;
 	virtual FReply OnKeyDown(const FGeometry& MyGeometry, const FKeyEvent& InKeyEvent) override;
-	virtual void Tick(const FGeometry& AllottedGeometry, const double InCurrentTime, const float InDeltaTime) override;
 	virtual FReply OnDragDetected(const FGeometry& MyGeometry, const FPointerEvent& MouseEvent) override;
 	virtual void OnDragEnter(const FGeometry& MyGeometry, const FDragDropEvent& DragDropEvent) override;
 	virtual void OnDragLeave(const FDragDropEvent& DragDropEvent) override;
 	virtual FReply OnDragOver(const FGeometry& MyGeometry, const FDragDropEvent& DragDropEvent) override;
 	virtual FReply OnDrop(const FGeometry& MyGeometry, const FDragDropEvent& DragDropEvent) override;
-	// End of SWidget interface
+	//~ End of SWidget interface
 
 	//~ Begin SDMXPixelMappingSurface interface
 	virtual FSlateRect ComputeAreaBounds() const override;
 	virtual int32 GetGraphRulePeriod() const override;
 	virtual float GetGridScaleAmount() const override;
-	virtual int32 GetSnapGridSize() const override;
+	virtual int32 GetGridSize() const override;
 	//~ End SDMXPixelMappingSurface interface
 
-	//~ Begin FEditorUndoClient Interface
+	//~ Begin FEditorUndoClient interface
 	virtual void PostUndo(bool bSuccess) override;
 	virtual void PostRedo(bool bSuccess) override;
-	// End of FEditorUndoClient
+	//~ End FEditorUndoClient interface
 
 private:
 	/** Rebuilds the designer widget */
@@ -126,7 +120,6 @@ private:
 	/** Returns the visibility of the ZoomPan widget */
 	EVisibility GetZoomPanVisibility() const;
 
-private:
 	/** Called when zoom to fit was clicked */
 	FReply OnZoomToFitClicked();
 
@@ -146,7 +139,6 @@ private:
 
 	void PopulateWidgetGeometryCache_Loop(FArrangedWidget& Parent);
 
-private:
 	/** Returns the component under the cursor */
 	UDMXPixelMappingOutputComponent* GetComponentUnderCursor() const;
 
@@ -172,7 +164,6 @@ private:
 	/** Updates the drag drop op to use the patches selected in the Palette */
 	void HandleDragEnterFromDetailsOrPalette(const TSharedPtr<FDMXPixelMappingDragDropOp>& DragDropOp);
 
-private:
 	/** The component that should be selected, but isn't selected yet */
 	TWeakObjectPtr<UDMXPixelMappingBaseComponent> PendingSelectedComponent;
 
@@ -188,19 +179,20 @@ private:
 	/** Canvas that holds the component widgets */
 	TSharedPtr<SConstraintCanvas> DesignCanvas;
 
-	TSharedPtr<SBox> PreviewSizeConstraint;
-
 	TSharedPtr<SOverlay> PreviewHitTestRoot;
 
 	/** Borer that contains the Design Canvas */
 	TSharedPtr<SBorder> DesignCanvasBorder;
 
-	TSharedPtr<FHittestGrid> HittestGrid;
+	/** Cached cursor position in graph space */
+	FVector2D CursorPositionGraphSpace = FVector2D::ZeroVector;
 
 	TMap<TWeakPtr<SWidget>, FArrangedWidget> CachedWidgetGeometry;
 
+	/** Cached current renderer component */
 	TWeakObjectPtr<UDMXPixelMappingRendererComponent> CachedRendererComponent;
 
+	/** The widget that handles zoom and pan */
 	TSharedPtr<SDMXPixelMappingZoomPan> ZoomPan;
 
 	/** The ruler bar at the top of the designer. */
@@ -212,7 +204,7 @@ private:
 	/** The position in local space where the user began dragging a widget */
 	FVector2D DragAnchor;
 
-	/** Oupput component views currently displayed */
+	/** Output component views currently displayed */
 	TArray<TSharedRef<IDMXPixelMappingOutputComponentWidgetInterface>> OutputComponentWidgets;
 
 	TArray<TSharedPtr<SDMXPixelMappingTransformHandle>> TransformHandles;
diff --git a/Engine/Plugins/VirtualProduction/DMX/DMXPixelMapping/Source/DMXPixelMappingEditor/Private/Views/SDMXPixelMappingPreviewView.cpp b/Engine/Plugins/VirtualProduction/DMX/DMXPixelMapping/Source/DMXPixelMappingEditor/Private/Views/SDMXPixelMappingPreviewView.cpp
--- a/Engine/Plugins/VirtualProduction/DMX/DMXPixelMapping/Source/DMXPixelMappingEditor/Private/Views/SDMXPixelMappingPreviewView.cpp
+++ b/Engine/Plugins/VirtualProduction/DMX/DMXPixelMapping/Source/DMXPixelMappingEditor/Private/Views/SDMXPixelMappingPreviewView.cpp
@@ -216,7 +216,7 @@ float SDMXPixelMappingPreviewView::GetGridScaleAmount() const
 	return 1.0f; // Parent override
 }
 
-int32 SDMXPixelMappingPreviewView::GetSnapGridSize() const
+int32 SDMXPixelMappingPreviewView::GetGridSize() const
 {
 	return 4; // Parent override
 }
diff --git a/Engine/Plugins/VirtualProduction/DMX/DMXPixelMapping/Source/DMXPixelMappingEditor/Private/Views/SDMXPixelMappingPreviewView.h b/Engine/Plugins/VirtualProduction/DMX/DMXPixelMapping/Source/DMXPixelMappingEditor/Private/Views/SDMXPixelMappingPreviewView.h
--- a/Engine/Plugins/VirtualProduction/DMX/DMXPixelMapping/Source/DMXPixelMappingEditor/Private/Views/SDMXPixelMappingPreviewView.h
+++ b/Engine/Plugins/VirtualProduction/DMX/DMXPixelMapping/Source/DMXPixelMappingEditor/Private/Views/SDMXPixelMappingPreviewView.h
@@ -51,7 +51,7 @@ protected:
 	virtual void OnPaintBackground(const FGeometry& AllottedGeometry, const FSlateRect& MyCullingRect, FSlateWindowElementList& OutDrawElements, int32 LayerId) const override;
 	virtual int32 GetGraphRulePeriod() const override;
 	virtual float GetGridScaleAmount() const override;
-	virtual int32 GetSnapGridSize() const override;
+	virtual int32 GetGridSize() const override;
 	virtual FSlateRect ComputeAreaBounds() const override;
 	//~ End SDMXPixelMappingSurface interface
 
diff --git a/Engine/Plugins/VirtualProduction/DMX/DMXPixelMapping/Source/DMXPixelMappingEditor/Private/Widgets/SDMXPixelMappingDesignerToolbar.cpp b/Engine/Plugins/VirtualProduction/DMX/DMXPixelMapping/Source/DMXPixelMappingEditor/Private/Widgets/SDMXPixelMappingDesignerToolbar.cpp
new file mode 100644
--- /dev/null
+++ b/Engine/Plugins/VirtualProduction/DMX/DMXPixelMapping/Source/DMXPixelMappingEditor/Private/Widgets/SDMXPixelMappingDesignerToolbar.cpp
@@ -0,0 +1,137 @@
+// Copyright Epic Games, Inc. All Rights Reserved.
+
+#include "SDMXPixelMappingDesignerToolbar.h"
+
+#include "DMXPixelMapping.h"
+#include "DMXPixelMappingEditorCommands.h"
+#include "Styling/AppStyle.h"
+#include "Styling/ToolBarStyle.h"
+#include "SViewportToolBarComboMenu.h"
+#include "Toolkits/DMXPixelMappingToolkit.h"
+#include "Widgets/Images/SImage.h"
+#include "Widgets/Input/SButton.h"
+#include "Widgets/Layout/SBorder.h"
+#include "Widgets/Text/STextBlock.h"
+#include "Widgets/SDMXPixelMappingSnapGridMenu.h"
+
+
+#define LOCTEXT_NAMESPACE "SDMXPixelMappingDesignerToolbar"
+
+namespace UE::DMX
+{
+	void SDMXPixelMappingDesignerToolbar::Construct(const FArguments& InArgs, const TSharedRef<FDMXPixelMappingToolkit>& InToolkit)
+	{
+		WeakToolkit = InToolkit;
+
+		constexpr TCHAR ToolBarStyleName[] = TEXT("EditorViewportToolBar");
+		const FToolBarStyle& ToolBarStyle = FAppStyle::Get().GetWidgetStyle<FToolBarStyle>(ToolBarStyleName);
+
+		const TSharedRef<FUICommandList> CommandList = InToolkit->GetToolkitCommands();
+
+		const FUICommandInfo* ToggleSnapGridCommand = FDMXPixelMappingEditorCommands::Get().ToggleGridSnapping.Get();
+		check(ToggleSnapGridCommand);
+
+		ChildSlot
+		[
+			SNew(SBorder)
+			.BorderImage(FAppStyle::GetBrush("NoBorder"))
+			.ForegroundColor(FAppStyle::GetSlateColor("DefaultForeground"))
+			[
+				SNew(SHorizontalBox)
+
+				// Grid snapping
+				+ SHorizontalBox::Slot()
+				.AutoWidth()
+				.VAlign(VAlign_Center)
+				.Padding(4.f, 2.0f)
+				[
+					SNew(SViewportToolBarComboMenu)
+					.Cursor(EMouseCursor::Default)
+					.IsChecked(this, &SDMXPixelMappingDesignerToolbar::GetSnapGridEnabledCheckState)
+					.OnCheckStateChanged(this, &SDMXPixelMappingDesignerToolbar::OnSnapGridCheckStateChanged)
+					.Label(this, &SDMXPixelMappingDesignerToolbar::GetSnapGridLabel)
+					.OnGetMenuContent(this, &SDMXPixelMappingDesignerToolbar::GenerateSnapGridMenu)
+					.ToggleButtonToolTip(ToggleSnapGridCommand->GetDescription())
+					.MenuButtonToolTip(LOCTEXT("SnapGridMenuTooltip", "Grid Snapping Settings"))
+					.Icon(FSlateIcon(FAppStyle::GetAppStyleSetName(), "EditorViewport.LocationGridSnap"))
+					.ParentToolBar(SharedThis(this))
+				]
+
+				// Zoom to fit
+				+ SHorizontalBox::Slot()
+				.AutoWidth()
+				.VAlign(VAlign_Center)
+				.Padding(4.f, 2.0f)
+				[
+					SNew(SButton)
+					.ButtonStyle(&ToolBarStyle.ButtonStyle)
+					.ToolTipText(LOCTEXT("ZoomToFit_ToolTip", "Zoom To Fit"))
+					.OnClicked(InArgs._OnZoomToFitClicked)
+					.ContentPadding(ToolBarStyle.ButtonPadding)
+					.VAlign(VAlign_Center)
+					[
+						SNew(SImage)
+						.Image(FAppStyle::GetBrush("UMGEditor.ZoomToFit"))
+						.ColorAndOpacity(FSlateColor::UseForeground())
+					]
+				]
+			]
+		];
+	}
+
+	TSharedRef<SWidget> SDMXPixelMappingDesignerToolbar::GenerateSnapGridMenu()
+	{
+		if (!WeakToolkit.IsValid())
+		{
+			return SNullWidget::NullWidget;
+		}
+
+		const TSharedRef<SWidget> SnapGridMenu = 
+			SNew(SDMXPixelMappingSnapGridMenu, WeakToolkit.Pin().ToSharedRef())
+			.IsChecked(this, &SDMXPixelMappingDesignerToolbar::GetSnapGridEnabledCheckState)
+			.OnCheckStateChanged(this, &SDMXPixelMappingDesignerToolbar::OnSnapGridCheckStateChanged);
+
+		return SnapGridMenu;
+	}
+
+	FText SDMXPixelMappingDesignerToolbar::GetSnapGridLabel() const
+	{
+		UDMXPixelMapping* PixelMapping = WeakToolkit.IsValid() ? WeakToolkit.Pin()->GetDMXPixelMapping() : nullptr;
+		if (PixelMapping)
+		{
+			if (PixelMapping->bGridSnappingEnabled)
+			{
+				const FString ColumnString = FString::FromInt(PixelMapping->SnapGridColumns);
+				const FString RowString = FString::FromInt(PixelMapping->SnapGridRows);
+
+				return FText::FromString(ColumnString + TEXT("x") + RowString);
+			}
+
+			return FText::GetEmpty();
+		}
+
+		return FText::GetEmpty();
+	}
+
+	ECheckBoxState SDMXPixelMappingDesignerToolbar::GetSnapGridEnabledCheckState() const
+	{
+		UDMXPixelMapping* PixelMapping = WeakToolkit.IsValid() ? WeakToolkit.Pin()->GetDMXPixelMapping() : nullptr;
+		if (PixelMapping)
+		{
+			return PixelMapping->bGridSnappingEnabled ? ECheckBoxState::Checked : ECheckBoxState::Unchecked;
+		}
+
+		return ECheckBoxState::Undetermined;
+	}
+
+	void SDMXPixelMappingDesignerToolbar::OnSnapGridCheckStateChanged(ECheckBoxState NewCheckBoxState)
+	{
+		if (const TSharedPtr<FDMXPixelMappingToolkit> Toolkit = WeakToolkit.Pin())
+		{
+			Toolkit->ToggleGridSnapping();
+		}
+	}
+}
+
+
+#undef LOCTEXT_NAMESPACE
diff --git a/Engine/Plugins/VirtualProduction/DMX/DMXPixelMapping/Source/DMXPixelMappingEditor/Private/Widgets/SDMXPixelMappingDesignerToolbar.h b/Engine/Plugins/VirtualProduction/DMX/DMXPixelMapping/Source/DMXPixelMappingEditor/Private/Widgets/SDMXPixelMappingDesignerToolbar.h
new file mode 100644
--- /dev/null
+++ b/Engine/Plugins/VirtualProduction/DMX/DMXPixelMapping/Source/DMXPixelMappingEditor/Private/Widgets/SDMXPixelMappingDesignerToolbar.h
@@ -0,0 +1,44 @@
+// Copyright Epic Games, Inc. All Rights Reserved.
+
+#pragma once
+
+#include "Framework/SlateDelegates.h"
+#include "Styling/SlateTypes.h"
+#include "SViewportToolBar.h"
+
+class FDMXPixelMappingToolkit;
+
+
+namespace UE::DMX
+{
+	class SDMXPixelMappingDesignerToolbar
+		: public SViewportToolBar
+	{
+	public:
+		SLATE_BEGIN_ARGS(SDMXPixelMappingDesignerToolbar)
+			{}
+
+			SLATE_EVENT(FOnClicked, OnZoomToFitClicked)
+
+		SLATE_END_ARGS()
+
+		/** Constructs this widget */
+		void Construct(const FArguments& InArgs, const TSharedRef<FDMXPixelMappingToolkit>& InToolkit);
+
+	private:
+		/** Generates the grid snapping menu */
+		TSharedRef<SWidget> GenerateSnapGridMenu();
+
+		/** Returns the grid snapping label */
+		FText GetSnapGridLabel() const;
+
+		/** Gets the check box state according to grid snapping being enabled */
+		ECheckBoxState GetSnapGridEnabledCheckState() const;
+
+		/** Called when the grid snapping checkbox state changed */
+		void OnSnapGridCheckStateChanged(ECheckBoxState NewCheckBoxState);
+
+		/** The toolkit that owns this widget */
+		TWeakPtr<FDMXPixelMappingToolkit> WeakToolkit;
+	};
+}
diff --git a/Engine/Plugins/VirtualProduction/DMX/DMXPixelMapping/Source/DMXPixelMappingEditor/Private/Widgets/SDMXPixelMappingFixturePatchList.cpp b/Engine/Plugins/VirtualProduction/DMX/DMXPixelMapping/Source/DMXPixelMappingEditor/Private/Widgets/SDMXPixelMappingFixturePatchList.cpp
--- a/Engine/Plugins/VirtualProduction/DMX/DMXPixelMapping/Source/DMXPixelMappingEditor/Private/Widgets/SDMXPixelMappingFixturePatchList.cpp
+++ b/Engine/Plugins/VirtualProduction/DMX/DMXPixelMapping/Source/DMXPixelMappingEditor/Private/Widgets/SDMXPixelMappingFixturePatchList.cpp
@@ -70,7 +70,7 @@ FReply SDMXPixelMappingFixturePatchList::OnRowDragDetected(const FGeometry& MyGe
 	}
 	else
 	{
-		return FReply::Handled().BeginDragDrop(FDMXPixelMappingDragDropOp::New(FVector2D::ZeroVector, Templates, FixtureGroup));
+		return FReply::Handled().BeginDragDrop(FDMXPixelMappingDragDropOp::New(Toolkit.ToSharedRef(), FVector2D::ZeroVector, Templates, FixtureGroup));
 	}
 }
 
diff --git a/Engine/Plugins/VirtualProduction/DMX/DMXPixelMapping/Source/DMXPixelMappingEditor/Private/Widgets/SDMXPixelMappingHierarchyRow.cpp b/Engine/Plugins/VirtualProduction/DMX/DMXPixelMapping/Source/DMXPixelMappingEditor/Private/Widgets/SDMXPixelMappingHierarchyRow.cpp
--- a/Engine/Plugins/VirtualProduction/DMX/DMXPixelMapping/Source/DMXPixelMappingEditor/Private/Widgets/SDMXPixelMappingHierarchyRow.cpp
+++ b/Engine/Plugins/VirtualProduction/DMX/DMXPixelMapping/Source/DMXPixelMappingEditor/Private/Widgets/SDMXPixelMappingHierarchyRow.cpp
@@ -220,7 +220,7 @@ FReply SDMXPixelMappingHierarchyRow::OnRowDragDetected(const FGeometry& MyGeomet
 		}
 	);
 
-	return FReply::Handled().BeginDragDrop(FDMXPixelMappingDragDropOp::New(FVector2D::ZeroVector, DraggedComponents));
+	return FReply::Handled().BeginDragDrop(FDMXPixelMappingDragDropOp::New(WeakToolkit.Pin().ToSharedRef(), FVector2D::ZeroVector, DraggedComponents));
 }
 
 #undef LOCTEXT_NAMESPACE
diff --git a/Engine/Plugins/VirtualProduction/DMX/DMXPixelMapping/Source/DMXPixelMappingEditor/Private/Widgets/SDMXPixelMappingSnapGridMenu.cpp b/Engine/Plugins/VirtualProduction/DMX/DMXPixelMapping/Source/DMXPixelMappingEditor/Private/Widgets/SDMXPixelMappingSnapGridMenu.cpp
new file mode 100644
--- /dev/null
+++ b/Engine/Plugins/VirtualProduction/DMX/DMXPixelMapping/Source/DMXPixelMappingEditor/Private/Widgets/SDMXPixelMappingSnapGridMenu.cpp
@@ -0,0 +1,273 @@
+// Copyright Epic Games, Inc. All Rights Reserved.
+
+#include "Widgets/SDMXPixelMappingSnapGridMenu.h"
+
+#include "DMXPixelMapping.h"
+#include "ScopedTransaction.h"
+#include "Styling/AppStyle.h"
+#include "Toolkits/DMXPixelMappingToolkit.h"
+#include "Widgets/Colors/SColorBlock.h"
+#include "Widgets/Colors/SColorPicker.h"
+#include "Widgets/Input/SButton.h"
+#include "Widgets/Input/SNumericEntryBox.h"
+#include "Widgets/Layout/SBorder.h"
+#include "Widgets/Layout/SGridPanel.h"
+#include "Widgets/SBoxPanel.h"
+#include "Widgets/Text/STextBlock.h"
+
+
+#define LOCTEXT_NAMESPACE "SDMXPixelMappingSnapGridMenu"
+
+namespace UE::DMX
+{
+	void SDMXPixelMappingSnapGridMenu::Construct(const FArguments& InArgs, const TSharedRef<FDMXPixelMappingToolkit>& InToolkit)
+	{
+		WeakToolkit = InToolkit;
+
+		ChildSlot
+		[
+			SNew(SBorder)
+			.BorderImage(FAppStyle::GetBrush("NoBorder"))
+			.ForegroundColor(FAppStyle::GetSlateColor("DefaultForeground"))
+			.Padding(12.f)
+			[
+				SNew(SGridPanel)
+				
+				// Enable grid snapping
+				+ SGridPanel::Slot(0, 0)
+				.HAlign(HAlign_Left)
+				.VAlign(VAlign_Center)
+				.Padding(4.f)
+				[
+					SNew(STextBlock)
+					.Text(LOCTEXT("SnapGridEnabledLabel", "Enable Grid Snapping"))
+					.Font(FAppStyle::Get().GetFontStyle("PropertyWindow.NormalFont"))
+				]
+
+				+ SGridPanel::Slot(1, 0)
+				.HAlign(HAlign_Left)
+				.VAlign(VAlign_Center)
+				[
+					SNew(SCheckBox)
+					.IsChecked(InArgs._IsChecked)
+					.OnCheckStateChanged(InArgs._OnCheckStateChanged)
+				]
+
+				// Num Rows
+				+ SGridPanel::Slot(0, 1)
+				.HAlign(HAlign_Left)
+				.VAlign(VAlign_Center)
+				.Padding(4.f)
+				[
+					SNew(STextBlock)
+					.Text(LOCTEXT("SnapGridColumnsLabel", "Num Columns (X)"))
+					.Font(FAppStyle::Get().GetFontStyle("PropertyWindow.NormalFont"))
+					.IsEnabled(this, &SDMXPixelMappingSnapGridMenu::IsGridSnappingEnabled)
+				]
+
+				+ SGridPanel::Slot(1, 1)
+				.HAlign(HAlign_Left)
+				.VAlign(VAlign_Center)
+				[
+					SNew(SNumericEntryBox<int32>)
+					.MinDesiredValueWidth(60.f)
+					.AllowSpin(true)
+					.MinValue(1)
+					.MaxValue(32768)
+					.MinSliderValue(1)
+					.MaxSliderValue(512)
+					.Value(this, &SDMXPixelMappingSnapGridMenu::GetColumns)
+					.OnValueChanged(this, &SDMXPixelMappingSnapGridMenu::OnColumnsEdited)
+					.IsEnabled(this, &SDMXPixelMappingSnapGridMenu::IsGridSnappingEnabled)
+				]
+
+				// Num Columns
+				+ SGridPanel::Slot(0, 2)
+				.HAlign(HAlign_Left)
+				.VAlign(VAlign_Center)
+				.Padding(4.f)
+				[
+					SNew(STextBlock)
+					.Text(LOCTEXT("SnapGridRowsLabel", "Num Rows (Y)"))
+					.Font(FAppStyle::Get().GetFontStyle("PropertyWindow.NormalFont"))
+					.IsEnabled(this, &SDMXPixelMappingSnapGridMenu::IsGridSnappingEnabled)
+				]
+
+				+ SGridPanel::Slot(1, 2)
+				.HAlign(HAlign_Left)
+				.VAlign(VAlign_Center)
+				[
+					SNew(SNumericEntryBox<int32>)
+					.MinDesiredValueWidth(60.f)
+					.AllowSpin(true)
+					.MinValue(1)
+					.MaxValue(32768)
+					.MinSliderValue(1)
+					.MaxSliderValue(512)
+					.Value(this, &SDMXPixelMappingSnapGridMenu::GetRows)
+					.OnValueChanged(this, &SDMXPixelMappingSnapGridMenu::OnRowsEdited)
+					.IsEnabled(this, &SDMXPixelMappingSnapGridMenu::IsGridSnappingEnabled)
+				]
+
+				// Color picker
+				+ SGridPanel::Slot(0, 3)
+				.HAlign(HAlign_Left)
+				.VAlign(VAlign_Center)
+				.Padding(4.f)
+				[
+					SNew(STextBlock)
+					.Text(LOCTEXT("ColorPickerLabel", "Color"))
+					.Font(FAppStyle::Get().GetFontStyle("PropertyWindow.NormalFont"))
+					.IsEnabled(this, &SDMXPixelMappingSnapGridMenu::IsGridSnappingEnabled)
+				]
+
+				+ SGridPanel::Slot(1, 3)
+				.HAlign(HAlign_Left)
+				.VAlign(VAlign_Center)
+				.Padding(4.f)
+				[
+					SNew(SBox)
+					.MinDesiredWidth(60.f)
+					[
+						SNew(SColorBlock)
+						.AlphaBackgroundBrush(FAppStyle::Get().GetBrush("ColorPicker.RoundedAlphaBackground"))
+						.Color(this, &SDMXPixelMappingSnapGridMenu::GetColor)
+						.ShowBackgroundForAlpha(true)
+						.AlphaDisplayMode(EColorBlockAlphaDisplayMode::Separate)
+						.OnMouseButtonDown(this, &SDMXPixelMappingSnapGridMenu::OnColorBlockClicked)
+						.Size(FVector2D(70.0f, 20.0f))
+						.CornerRadius(FVector4(2.0f, 2.0f, 2.0f, 2.0f))
+						.IsEnabled(this, &SDMXPixelMappingSnapGridMenu::IsGridSnappingEnabled)
+					]
+				]
+			]
+		];
+	}
+
+
+	TOptional<int32> SDMXPixelMappingSnapGridMenu::GetColumns() const
+	{
+		const UDMXPixelMapping* PixelMapping = WeakToolkit.IsValid() ? WeakToolkit.Pin()->GetDMXPixelMapping() : nullptr;
+		if (PixelMapping)
+		{
+			return PixelMapping->SnapGridColumns;
+		}
+
+		return 1;
+	}
+
+	void SDMXPixelMappingSnapGridMenu::OnColumnsEdited(int32 NewColumns)
+	{
+		UDMXPixelMapping* PixelMapping = WeakToolkit.IsValid() ? WeakToolkit.Pin()->GetDMXPixelMapping() : nullptr;
+		if (PixelMapping)
+		{
+			const FScopedTransaction SetSnapGridColumnsTransation(LOCTEXT("SetSnapGridColumnsTransation", "Set Pixel Mapping Grid X"));
+
+			PixelMapping->PreEditChange(UDMXPixelMapping::StaticClass()->FindPropertyByName(GET_MEMBER_NAME_CHECKED(UDMXPixelMapping, SnapGridColumns)));
+			PixelMapping->SnapGridColumns = NewColumns;
+			PixelMapping->PostEditChange();
+		}
+	}
+
+	TOptional<int32> SDMXPixelMappingSnapGridMenu::GetRows() const
+	{
+		const UDMXPixelMapping* PixelMapping = WeakToolkit.IsValid() ? WeakToolkit.Pin()->GetDMXPixelMapping() : nullptr;
+		if (PixelMapping)
+		{
+			return PixelMapping->SnapGridRows;
+		}
+
+		return 1;
+	}
+
+	void SDMXPixelMappingSnapGridMenu::OnRowsEdited(int32 NewRows)
+	{
+		UDMXPixelMapping* PixelMapping = WeakToolkit.IsValid() ? WeakToolkit.Pin()->GetDMXPixelMapping() : nullptr;
+		if (PixelMapping)
+		{
+			const FScopedTransaction SetSnapGridRowTransation(LOCTEXT("SetSnapGridRowTransation", "Set Pixel Mapping Grid X"));
+
+			PixelMapping->PreEditChange(UDMXPixelMapping::StaticClass()->FindPropertyByName(GET_MEMBER_NAME_CHECKED(UDMXPixelMapping, SnapGridRows)));
+			PixelMapping->SnapGridRows = NewRows;
+			PixelMapping->PostEditChange();
+		}
+	}
+
+	FLinearColor SDMXPixelMappingSnapGridMenu::GetColor() const
+	{
+		UDMXPixelMapping* PixelMapping = WeakToolkit.IsValid() ? WeakToolkit.Pin()->GetDMXPixelMapping() : nullptr;
+		if (PixelMapping)
+		{
+			return PixelMapping->SnapGridColor;
+		}
+
+		return FLinearColor::Red;
+	}
+
+	FReply SDMXPixelMappingSnapGridMenu::OnColorBlockClicked(const FGeometry& MyGeometry, const FPointerEvent& MouseEvent)
+	{
+		UDMXPixelMapping* PixelMapping = WeakToolkit.IsValid() ? WeakToolkit.Pin()->GetDMXPixelMapping() : nullptr;
+		if (!PixelMapping)
+		{
+			return FReply::Handled();
+		}
+
+		FColorPickerArgs PickerArgs;
+		{
+			PickerArgs.bUseAlpha = true;
+			PickerArgs.bOnlyRefreshOnMouseUp = false;
+			PickerArgs.bOnlyRefreshOnOk = false;
+			PickerArgs.DisplayGamma = TAttribute<float>::Create(TAttribute<float>::FGetter::CreateUObject(GEngine, &UEngine::GetDisplayGamma));
+			PickerArgs.OnColorCommitted = FOnLinearColorValueChanged::CreateSP(this, &SDMXPixelMappingSnapGridMenu::OnColorSelected);
+			PickerArgs.OnColorPickerCancelled = FOnColorPickerCancelled::CreateSP(this, &SDMXPixelMappingSnapGridMenu::OnColorPickerCancelled);
+			PickerArgs.OnColorPickerWindowClosed = FOnWindowClosed::CreateSP(this, &SDMXPixelMappingSnapGridMenu::OnColorPickerWindowClosed);
+			PickerArgs.InitialColor = PixelMapping->SnapGridColor;
+			PickerArgs.ParentWidget = AsShared();
+			PickerArgs.bOpenAsMenu = true;
+			PickerArgs.bIsModal = false;
+		}
+
+		ColorPickerTransaction = MakeShared<FScopedTransaction>(LOCTEXT("ColorPickerTransaction", "Set Pixel Mapping Grid Color"));
+
+		OpenColorPicker(PickerArgs);
+		
+		return FReply::Handled();
+	}
+
+	void SDMXPixelMappingSnapGridMenu::OnColorSelected(FLinearColor NewColor)
+	{
+		UDMXPixelMapping* PixelMapping = WeakToolkit.IsValid() ? WeakToolkit.Pin()->GetDMXPixelMapping() : nullptr;
+		if (PixelMapping)
+		{
+			PixelMapping->SnapGridColor = NewColor;
+		}
+	}
+
+	void SDMXPixelMappingSnapGridMenu::OnColorPickerCancelled(FLinearColor OriginalColor)
+	{
+		UDMXPixelMapping* PixelMapping = WeakToolkit.IsValid() ? WeakToolkit.Pin()->GetDMXPixelMapping() : nullptr;
+		if (PixelMapping)
+		{
+			PixelMapping->SnapGridColor = OriginalColor;
+		}
+	}
+
+	void SDMXPixelMappingSnapGridMenu::OnColorPickerWindowClosed(const TSharedRef<SWindow>& Window)
+	{
+		// Complete the transaction
+		ColorPickerTransaction.Reset();
+	}
+
+	bool SDMXPixelMappingSnapGridMenu::IsGridSnappingEnabled() const
+	{
+		UDMXPixelMapping* PixelMapping = WeakToolkit.IsValid() ? WeakToolkit.Pin()->GetDMXPixelMapping() : nullptr;
+		if (PixelMapping)
+		{
+			return PixelMapping->bGridSnappingEnabled;
+		}
+
+		return false;
+	}
+}
+
+#undef LOCTEXT_NAMESPACE
diff --git a/Engine/Plugins/VirtualProduction/DMX/DMXPixelMapping/Source/DMXPixelMappingEditor/Private/Widgets/SDMXPixelMappingSnapGridMenu.h b/Engine/Plugins/VirtualProduction/DMX/DMXPixelMapping/Source/DMXPixelMappingEditor/Private/Widgets/SDMXPixelMappingSnapGridMenu.h
new file mode 100644
--- /dev/null
+++ b/Engine/Plugins/VirtualProduction/DMX/DMXPixelMapping/Source/DMXPixelMappingEditor/Private/Widgets/SDMXPixelMappingSnapGridMenu.h
@@ -0,0 +1,79 @@
+// Copyright Epic Games, Inc. All Rights Reserved.
+
+#pragma once
+
+#include "Framework/SlateDelegates.h"
+#include "Input/Reply.h"
+#include "Misc/Optional.h"
+#include "SViewportToolBar.h"
+#include "Widgets/Input/SCheckBox.h"
+#include "Widgets/SCompoundWidget.h"
+
+class FDMXPixelMappingToolkit;
+class FScopedTransaction;
+class SWindow;
+
+
+namespace UE::DMX
+{
+	class SDMXPixelMappingSnapGridMenu
+		: public SCompoundWidget
+	{
+	public:
+		SLATE_BEGIN_ARGS(SDMXPixelMappingSnapGridMenu)
+			{}
+
+			/** Called upon state change with the value of the next state */
+			SLATE_EVENT(FOnCheckStateChanged, OnCheckStateChanged)
+
+			/** Sets the current checked state of the checkbox */
+			SLATE_ATTRIBUTE(ECheckBoxState, IsChecked)
+
+		SLATE_END_ARGS()
+
+	public:
+		/** Constructs this widget */
+		void Construct(const FArguments& InArgs, const TSharedRef<FDMXPixelMappingToolkit>& InToolkit);
+
+	private:
+		/** Returns the number of columns for grid snapping as text */
+		TOptional<int32> GetColumns() const;
+
+		/** Called when the number of grid snapping columns was edited */
+		void OnColumnsEdited(int32 NewColumns);
+
+		/** Returns the number of rows for grid snapping as text */
+		TOptional<int32> GetRows() const;
+
+		/** Called when the number of grid snapping rows was edited */
+		void OnRowsEdited(int32 NewRows);
+
+		/** Returns the current grid color */
+		FLinearColor GetColor() const;
+
+		/** Called when the grid color block was clicked */
+		FReply OnColorBlockClicked(const FGeometry& MyGeometry, const FPointerEvent& MouseEvent);
+
+		/** Called when a color was selected in the color picker */
+		void OnColorSelected(FLinearColor NewColor);
+
+		/** Called when color picking was canceled */
+		void OnColorPickerCancelled(FLinearColor OriginalColor);
+
+		/** Called when the color picker window was closed */ 
+		void OnColorPickerWindowClosed(const TSharedRef<SWindow>& Window);
+
+		/** Returns true if grid snapping is enabled */
+		bool IsGridSnappingEnabled() const;
+
+		/** Transaction when selecting colors in its own menu */
+		TSharedPtr<FScopedTransaction> ColorPickerTransaction;
+
+		/** The toolkit that owns this widget */
+		TWeakPtr<FDMXPixelMappingToolkit> WeakToolkit;
+
+		// Slate args
+		FOnCheckStateChanged OnCheckStateChanged;
+		ECheckBoxState IsChecked;
+	};
+}
\ No newline at end of file
diff --git a/Engine/Plugins/VirtualProduction/DMX/DMXPixelMapping/Source/DMXPixelMappingEditor/Private/Widgets/SDMXPixelMappingSurface.cpp b/Engine/Plugins/VirtualProduction/DMX/DMXPixelMapping/Source/DMXPixelMappingEditor/Private/Widgets/SDMXPixelMappingSurface.cpp
--- a/Engine/Plugins/VirtualProduction/DMX/DMXPixelMapping/Source/DMXPixelMappingEditor/Private/Widgets/SDMXPixelMappingSurface.cpp
+++ b/Engine/Plugins/VirtualProduction/DMX/DMXPixelMapping/Source/DMXPixelMappingEditor/Private/Widgets/SDMXPixelMappingSurface.cpp
@@ -2,9 +2,10 @@
 
 #include "Widgets/SDMXPixelMappingSurface.h"
 
-#include "Framework/Application/SlateApplication.h"
 #include "Components/DMXPixelMappingRendererComponent.h"
+#include "Framework/Application/SlateApplication.h"
 #include "Settings/LevelEditorViewportSettings.h"
+#include "SNodePanel.h"
 #include "Toolkits/DMXPixelMappingToolkit.h"
 
 
@@ -240,7 +241,7 @@ int32 SDMXPixelMappingSurface::OnPaint(const FPaintArgs& Args, const FGeometry&
 	OnPaintBackground(AllottedGeometry, MyCullingRect, OutDrawElements, LayerId);
 
 	SCompoundWidget::OnPaint(Args, AllottedGeometry, MyCullingRect, OutDrawElements, LayerId, InWidgetStyle, bParentEnabled);
-
+ 
 	return LayerId;
 }
 
@@ -593,7 +594,7 @@ void SDMXPixelMappingSurface::PaintBackgroundAsLines(const FSlateBrush* Backgrou
 	const FLinearColor CenterColor(FAppStyle::GetColor("Graph.Panel.GridCenterColor"));
 	const float GraphSmallestGridSize = 8.0f;
 	const float RawZoomFactor = GetZoomAmount();
-	const float NominalGridSize = GetSnapGridSize() * GetGridScaleAmount();
+	const float NominalGridSize = GetGridSize() * GetGridScaleAmount();
 
 	float ZoomFactor = RawZoomFactor;
 	float Inflation = 1.0f;
diff --git a/Engine/Plugins/VirtualProduction/DMX/DMXPixelMapping/Source/DMXPixelMappingEditor/Private/Widgets/SDMXPixelMappingSurface.h b/Engine/Plugins/VirtualProduction/DMX/DMXPixelMapping/Source/DMXPixelMappingEditor/Private/Widgets/SDMXPixelMappingSurface.h
--- a/Engine/Plugins/VirtualProduction/DMX/DMXPixelMapping/Source/DMXPixelMappingEditor/Private/Widgets/SDMXPixelMappingSurface.h
+++ b/Engine/Plugins/VirtualProduction/DMX/DMXPixelMapping/Source/DMXPixelMappingEditor/Private/Widgets/SDMXPixelMappingSurface.h
@@ -2,15 +2,15 @@
 
 #pragma once
 
-#include "SNodePanel.h" // IWYU pragma: keep
-
-struct FZoomLevelsContainer;
-
+#include "Animation/CurveSequence.h"
+#include "SNodePanel.h"
+#include "Widgets/SCompoundWidget.h"
 
 class FActiveTimerHandle;
 class FDMXPixelMappingToolkit;
 class FSlateWindowElementList;
 
+
 /**
  * Most of the logic copied from a private class Engine/Source/Editor/UMGEditor/Private/Designer/SPaintSurface.h
  * Base class for Preview and Designer Grid
@@ -82,7 +82,7 @@ protected:
 	virtual FSlateRect ComputeAreaBounds() const;
 	virtual float GetGridScaleAmount() const;
 	virtual int32 GetGraphRulePeriod() const;
-	virtual int32 GetSnapGridSize() const = 0;
+	virtual int32 GetGridSize() const = 0;
 
 protected:
 	/** The position within the graph at which the user is looking */
diff --git a/Engine/Plugins/VirtualProduction/DMX/DMXPixelMapping/Source/DMXPixelMappingEditor/Private/Widgets/SDMXPixelMappingTransformHandle.cpp b/Engine/Plugins/VirtualProduction/DMX/DMXPixelMapping/Source/DMXPixelMappingEditor/Private/Widgets/SDMXPixelMappingTransformHandle.cpp
--- a/Engine/Plugins/VirtualProduction/DMX/DMXPixelMapping/Source/DMXPixelMappingEditor/Private/Widgets/SDMXPixelMappingTransformHandle.cpp
+++ b/Engine/Plugins/VirtualProduction/DMX/DMXPixelMapping/Source/DMXPixelMappingEditor/Private/Widgets/SDMXPixelMappingTransformHandle.cpp
@@ -3,6 +3,8 @@
 #include "Widgets/SDMXPixelMappingTransformHandle.h"
 
 #include "Components/DMXPixelMappingMatrixCellComponent.h"
+#include "Components/DMXPixelMappingRendererComponent.h"
+#include "DMXPixelMapping.h"
 #include "Framework/Application/SlateApplication.h"
 #include "ScopedTransaction.h"
 #include "Settings/DMXPixelMappingEditorSettings.h"
@@ -135,6 +137,23 @@ FReply SDMXPixelMappingTransformHandle::OnMouseMove(const FGeometry& MyGeometry,
 	return FReply::Unhandled();
 }
 
+FCursorReply SDMXPixelMappingTransformHandle::OnCursorQuery(const FGeometry& MyGeometry, const FPointerEvent& MouseEvent) const
+{
+	switch (TransformDirection)
+	{
+	case EDMXPixelMappingTransformDirection::BottomRight:
+		return FCursorReply::Cursor(EMouseCursor::ResizeSouthEast);
+	case EDMXPixelMappingTransformDirection::BottomLeft:
+		return FCursorReply::Cursor(EMouseCursor::ResizeSouthWest);
+	case EDMXPixelMappingTransformDirection::BottomCenter:
+		return FCursorReply::Cursor(EMouseCursor::ResizeUpDown);
+	case EDMXPixelMappingTransformDirection::CenterRight:
+		return FCursorReply::Cursor(EMouseCursor::ResizeLeftRight);
+	}
+
+	return FCursorReply::Unhandled();
+}
+
 void SDMXPixelMappingTransformHandle::RequestResize(UDMXPixelMappingBaseComponent* BaseComponent, const FVector2D& Direction)
 {
 	if (!RequestResizeHandle.IsValid() && DesignerViewWeakPtr.IsValid())
@@ -184,13 +203,14 @@ void SDMXPixelMappingTransformHandle::Resize(UDMXPixelMappingBaseComponent* Base
 		}
 
 		const FVector2D OldSize = OutputComponent->GetSize();
-		const FVector2D NewSize = FVector2D(Offsets.Right, Offsets.Bottom);
+		const FVector2D RequestedSize = FVector2D(Offsets.Right, Offsets.Bottom);
+		const FVector2D NewSize = GetSnapSize(OutputComponent, RequestedSize, Direction);
 		if (OldSize == NewSize)
 		{
 			// No unchanged values
 			return;
 		}
-
+		OutputComponent->Modify();
 		OutputComponent->SetSize(NewSize);
 
 		// Scale children only if desired, no division by zero
@@ -200,55 +220,85 @@ void SDMXPixelMappingTransformHandle::Resize(UDMXPixelMappingBaseComponent* Base
 			return;
 		}
 
-		const FVector2D RatioVector = NewSize / OldSize;
+		// Scale children with parent. 
+		FVector2D RatioVector = NewSize / OldSize;
 		for (UDMXPixelMappingBaseComponent* BaseChild : OutputComponent->GetChildren())
 		{
 			if (UDMXPixelMappingOutputComponent* Child = Cast<UDMXPixelMappingOutputComponent>(BaseChild))
 			{
-				if (BaseChild->GetClass() == UDMXPixelMappingMatrixCellComponent::StaticClass())
-				{
-					// Don't scale matrix cells, the matrix component already cares for this
-					if (Child->IsLockInDesigner())
-					{
-						continue;
-					}
-				}
-
 				Child->Modify();
 
-				// Scale size (Note, SetSize already clamps)
-				Child->SetSize(Child->GetSize() * RatioVector);
+				// Scale size, at least 1x1 pixel
+				FVector2D NewChildSize = Child->GetSize() * RatioVector;
+				NewChildSize.X = FMath::Max(NewChildSize.X, 1.f);
+				NewChildSize.Y = FMath::Max(NewChildSize.Y, 1.f);
+				Child->SetSize(NewChildSize);
 
 				// Scale position
-				const FVector2D ChildPosition = Child->GetPosition();
-				const FVector2D NewPositionRelative = (ChildPosition - OutputComponent->GetPosition()) * RatioVector;
-				Child->SetPosition(OutputComponent->GetPosition() + NewPositionRelative);
+				const FVector2D NewChildPosition = Child->GetPosition();
+				const FVector2D NewChildPositionRelative = (NewChildPosition - OutputComponent->GetPosition()) * RatioVector;
+				Child->SetPosition(OutputComponent->GetPosition() + NewChildPositionRelative);
 			}
 		}
 	}
 }
 
-FCursorReply SDMXPixelMappingTransformHandle::OnCursorQuery(const FGeometry& MyGeometry, const FPointerEvent& MouseEvent) const
+FVector2D SDMXPixelMappingTransformHandle::GetSnapSize(UDMXPixelMappingOutputComponent* OutputComponent, const FVector2D& RequestedSize, const FVector2D& Direction) const
 {
-	EDMXPixelMappingTransformAction CurrentAction = Action;
-	if ( CurrentAction == EDMXPixelMappingTransformAction::None )
+	const TSharedPtr<FDMXPixelMappingToolkit> Toolkit = DesignerViewWeakPtr.Pin()->GetToolkit();
+	if (!Toolkit.IsValid() || !OutputComponent || !OutputComponent->GetRendererComponent())
 	{
-		CurrentAction = ComputeActionAtLocation(MyGeometry, MouseEvent);
+		return RequestedSize;
 	}
+	UDMXPixelMappingRendererComponent* RendererComponent = OutputComponent->GetRendererComponent();
 
-	switch ( TransformDirection )
+	const UDMXPixelMapping* PixelMapping = Toolkit->GetDMXPixelMapping();
+	if (!PixelMapping)
 	{
-		case EDMXPixelMappingTransformDirection::BottomRight:
-			return FCursorReply::Cursor(EMouseCursor::ResizeSouthEast);
-		case EDMXPixelMappingTransformDirection::BottomLeft:
-			return FCursorReply::Cursor(EMouseCursor::ResizeSouthWest);
-		case EDMXPixelMappingTransformDirection::BottomCenter:
-			return FCursorReply::Cursor(EMouseCursor::ResizeUpDown);
-		case EDMXPixelMappingTransformDirection::CenterRight:
-			return FCursorReply::Cursor(EMouseCursor::ResizeLeftRight);
+		return RequestedSize;
 	}
 
-	return FCursorReply::Unhandled();
+	const FVector2D CellSize = [PixelMapping, RendererComponent]()
+		{
+			if (PixelMapping->bGridSnappingEnabled)
+			{
+				const FVector2D TextureSize = RendererComponent->GetSize();
+				return TextureSize / FVector2D(PixelMapping->SnapGridColumns, PixelMapping->SnapGridRows);
+			}
+
+			// Grid snap to pixels if grid snapping is disabled
+			return FVector2D(1.f, 1.f);
+		}();
+
+	// Grid snap bottom right
+	const int32 BottomColumn = (OutputComponent->GetPosition().X + RequestedSize.X) / CellSize.X + 1;
+	const int32 RightRow = (OutputComponent->GetPosition().Y + RequestedSize.Y) / CellSize.Y + 1;
+
+	FVector2D NewBottomRight = OutputComponent->GetPosition() + RequestedSize;
+	if (Direction.X < 0.f)
+	{
+		// Snap towards left
+		NewBottomRight.X = (BottomColumn - 1) * CellSize.X;
+	}
+	else if (Direction.X > 0.f)
+	{
+		// Snap towards right
+		NewBottomRight.X = BottomColumn * CellSize.X;
+	}
+
+	if (Direction.Y < 0.f)
+	{
+		// Snap towards top
+		NewBottomRight.Y = (RightRow - 1) * CellSize.Y;
+	}
+	else if (Direction.Y > 0.f)
+	{
+		// Snap towards bottom
+		NewBottomRight.Y = RightRow * CellSize.Y;
+	}
+
+	const FVector2D SnapSize = NewBottomRight - OutputComponent->GetPosition();
+	return SnapSize;
 }
 
 FVector2D SDMXPixelMappingTransformHandle::ComputeDragDirection(EDMXPixelMappingTransformDirection InTransformDirection) const
diff --git a/Engine/Plugins/VirtualProduction/DMX/DMXPixelMapping/Source/DMXPixelMappingEditor/Private/Widgets/SDMXPixelMappingTransformHandle.h b/Engine/Plugins/VirtualProduction/DMX/DMXPixelMapping/Source/DMXPixelMappingEditor/Private/Widgets/SDMXPixelMappingTransformHandle.h
--- a/Engine/Plugins/VirtualProduction/DMX/DMXPixelMapping/Source/DMXPixelMappingEditor/Private/Widgets/SDMXPixelMappingTransformHandle.h
+++ b/Engine/Plugins/VirtualProduction/DMX/DMXPixelMapping/Source/DMXPixelMappingEditor/Private/Widgets/SDMXPixelMappingTransformHandle.h
@@ -5,6 +5,10 @@
 #include "Engine/EngineTypes.h"
 #include "Widgets/SCompoundWidget.h"
 
+class FScopedTransaction;
+class SDMXPixelMappingDesignerView;
+class UDMXPixelMappingBaseComponent;
+class UDMXPixelMappingOutputComponent;
 
 enum class EDMXPixelMappingTransformDirection : uint8
 {
@@ -23,9 +27,6 @@ enum class EDMXPixelMappingTransformAction : uint8
 	Secondary
 };
 
-class SDMXPixelMappingDesignerView;
-class UDMXPixelMappingBaseComponent;
-class FScopedTransaction;
 
 /**
  * Most of the logic copied from a private class Engine/Source/Editor/UMGEditor/Private/Designer/STransformHandle.h
@@ -54,7 +55,7 @@ public:
 	EDMXPixelMappingTransformDirection GetTransformDirection() const { return TransformDirection; }
 	FVector2D GetOffset() const { return Offset.Get(); }
 
-protected:
+private:
 	EVisibility GetHandleVisibility() const;
 
 	/** Resizes the component on the next tick */
@@ -63,13 +64,14 @@ protected:
 	/** Resizes the component */
 	void Resize(UDMXPixelMappingBaseComponent* BaseComponent, const FVector2D& Direction, const FVector2D& Amount);
 
+	/** Utiltity that returns the snap size of an output component. Note, this snaps to pixels if grid snapping is disabled. */
+	FVector2D GetSnapSize(UDMXPixelMappingOutputComponent* OutputComponent, const FVector2D& RequestedSize, const FVector2D& Direction) const;
+
 	EDMXPixelMappingTransformAction ComputeActionAtLocation(const FGeometry& MyGeometry, const FPointerEvent& MouseEvent) const;
 
-protected:
 	FVector2D ComputeDragDirection(EDMXPixelMappingTransformDirection InTransformDirection) const;
 	FVector2D ComputeOrigin(EDMXPixelMappingTransformDirection InTransformDirection) const;
 
-protected:
 	EDMXPixelMappingTransformDirection TransformDirection;
 	EDMXPixelMappingTransformAction Action;
 
diff --git a/Engine/Plugins/VirtualProduction/DMX/DMXPixelMapping/Source/DMXPixelMappingRuntime/Private/Components/DMXPixelMappingOutputComponent.cpp b/Engine/Plugins/VirtualProduction/DMX/DMXPixelMapping/Source/DMXPixelMappingRuntime/Private/Components/DMXPixelMappingOutputComponent.cpp
--- a/Engine/Plugins/VirtualProduction/DMX/DMXPixelMapping/Source/DMXPixelMappingRuntime/Private/Components/DMXPixelMappingOutputComponent.cpp
+++ b/Engine/Plugins/VirtualProduction/DMX/DMXPixelMapping/Source/DMXPixelMappingRuntime/Private/Components/DMXPixelMappingOutputComponent.cpp
@@ -272,16 +272,8 @@ FVector2D UDMXPixelMappingOutputComponent::GetPosition() const
 
 void UDMXPixelMappingOutputComponent::SetSize(const FVector2D& Size) 
 {
-	SizeX = FMath::Max(Size.X, 1.f);
-	SizeY = FMath::Max(Size.Y, 1.f);
-
-	// Limit all components size to max texture size, logged
-	const uint32 MaxTextureDimensions = GetMax2DTextureDimension();
-	if (SizeX > MaxTextureDimensions || SizeY > MaxTextureDimensions)
-	{
-		SizeX = MaxTextureDimensions;
-		SizeY = MaxTextureDimensions;
-	}
+	SizeX = Size.X;
+	SizeY = Size.Y;
 }
 
 void UDMXPixelMappingOutputComponent::InvalidatePixelMapRenderer()
diff --git a/Engine/Plugins/VirtualProduction/DMX/DMXPixelMapping/Source/DMXPixelMappingRuntime/Private/DMXPixelMapping.cpp b/Engine/Plugins/VirtualProduction/DMX/DMXPixelMapping/Source/DMXPixelMappingRuntime/Private/DMXPixelMapping.cpp
--- a/Engine/Plugins/VirtualProduction/DMX/DMXPixelMapping/Source/DMXPixelMappingRuntime/Private/DMXPixelMapping.cpp
+++ b/Engine/Plugins/VirtualProduction/DMX/DMXPixelMapping/Source/DMXPixelMappingRuntime/Private/DMXPixelMapping.cpp
@@ -7,14 +7,20 @@
 #include "Components/DMXPixelMappingMatrixComponent.h"
 #include "Components/DMXPixelMappingMatrixCellComponent.h"
 #include "Library/DMXEntityFixturePatch.h"
+#include "UObject/LinkerLoad.h"
 
 #if WITH_EDITOR
 #include "DMXPixelMappingComponentWidget.h"
 #include "SDMXPixelMappingComponentBox.h"
 #endif // WITH_EDITOR
 
-#include "UObject/LinkerLoad.h"
 
+UDMXPixelMapping::UDMXPixelMapping()
+{
+#if WITH_EDITOR
+	SnapGridColor = FLinearColor::White.CopyWithNewOpacity(.12f);
+#endif // WITH_EDITOR
+}
 
 void UDMXPixelMapping::PostLoad()
 {
diff --git a/Engine/Plugins/VirtualProduction/DMX/DMXPixelMapping/Source/DMXPixelMappingRuntime/Public/DMXPixelMapping.h b/Engine/Plugins/VirtualProduction/DMX/DMXPixelMapping/Source/DMXPixelMappingRuntime/Public/DMXPixelMapping.h
--- a/Engine/Plugins/VirtualProduction/DMX/DMXPixelMapping/Source/DMXPixelMappingRuntime/Public/DMXPixelMapping.h
+++ b/Engine/Plugins/VirtualProduction/DMX/DMXPixelMapping/Source/DMXPixelMappingRuntime/Public/DMXPixelMapping.h
@@ -11,10 +11,6 @@ class UDMXPixelMappingBaseComponent;
 class UDMXPixelMappingOutputComponent;
 class UTexture;
 
-#if WITH_EDITOR
-DECLARE_DELEGATE_OneParam(FOnEditorRebuildChildrenComponentsDelegate, UDMXPixelMappingBaseComponent*)
-#endif // WITH_EDITOR
-
 
 /**
  * Public container of Pixel Mapping object and it using for asset
@@ -31,6 +27,8 @@ class DMXPIXELMAPPINGRUNTIME_API UDMXPixelMapping
 
 	GENERATED_BODY()
 public:
+	UDMXPixelMapping();
+
 	//~ Begin UObject implementation
 	virtual void PostLoad() override;
 	//~ End UObject implementation
@@ -144,13 +142,25 @@ public:
 	UPROPERTY()
 	TObjectPtr<UDMXPixelMappingRootComponent> RootComponent;
 
-
 #if WITH_EDITORONLY_DATA
-	/** Holds the Thumbnail for asset */
+	/** If true, grid snapping is enabled */
 	UPROPERTY()
-	TObjectPtr<UTexture> ThumbnailImage;
+	bool bGridSnappingEnabled = false;
+
+	/** The number of columns in the grid */
+	UPROPERTY()
+	int32 SnapGridColumns = 10;
+
+	/** The number of rows in the grid */
+	UPROPERTY()
+	int32 SnapGridRows = 10;
 
-	/** DEPRECATED 4.27. Use FDMXPixelMappingMatrixComponent::GetOnMatrixChanged() instead */
-	FOnEditorRebuildChildrenComponentsDelegate OnEditorRebuildChildrenComponentsDelegate_DEPRECATED;
-#endif
+	/** The color of the grid snapping grid */
+	UPROPERTY()
+	FLinearColor SnapGridColor;
+
+	/** Holds the Thumbnail image for this asset */
+	UPROPERTY()
+	TObjectPtr<UTexture> ThumbnailImage;
+#endif // WITH_EDITORONLY_DATA
 };
