DMX - PixelMapping now has an option to let component tranforms follow the texture size

diff --git a/Engine/Plugins/VirtualProduction/DMX/DMXPixelMapping/Source/DMXPixelMappingRuntime/Private/Components/DMXPixelMappingRendererComponent.cpp b/Engine/Plugins/VirtualProduction/DMX/DMXPixelMapping/Source/DMXPixelMappingRuntime/Private/Components/DMXPixelMappingRendererComponent.cpp
--- a/Engine/Plugins/VirtualProduction/DMX/DMXPixelMapping/Source/DMXPixelMappingRuntime/Private/Components/DMXPixelMappingRendererComponent.cpp
+++ b/Engine/Plugins/VirtualProduction/DMX/DMXPixelMapping/Source/DMXPixelMappingRuntime/Private/Components/DMXPixelMappingRendererComponent.cpp
@@ -38,22 +38,13 @@ DECLARE_CYCLE_STAT(TEXT("PixelMapping RenderInputTexture"), STAT_DMXPixelMapping
 
 UDMXPixelMappingRendererComponent::UDMXPixelMappingRendererComponent()
 {
-	SetSize(FVector2D(100.f, 100.f));
-	
-#if WITH_EDITOR
 	ConstructorHelpers::FObjectFinder<UTexture> DefaultTexture(TEXT("Texture2D'/Engine/VREditor/Devices/Vive/UE4_Logo.UE4_Logo'"), LOAD_NoWarn);
 	if (ensureAlwaysMsgf(DefaultTexture.Succeeded(), TEXT("Failed to load Texture2D'/Engine/VREditor/Devices/Vive/UE4_Logo.UE4_Logo'")))
 	{
 		InputTexture = DefaultTexture.Object;
 		RendererType = EDMXPixelMappingRendererType::Texture;
-
-		if (FTextureResource* Resource = InputTexture->GetResource())
-		{
-			SetSize(FVector2D(Resource->GetSizeX(), Resource->GetSizeY()));
-		}
 	}
-#endif
-	
+
 	PreprocessRenderer = CreateDefaultSubobject<UDMXPixelMappingPreprocessRenderer>("PreprocessRenderer");
 	PixelMapRenderer = CreateDefaultSubobject<UDMXPixelMappingPixelMapRenderer>("PixelMapRenderer");
 
@@ -84,6 +75,21 @@ void UDMXPixelMappingRendererComponent::PostInitProperties()
 	}
 }
 
+void UDMXPixelMappingRendererComponent::Serialize(FArchive& Ar)
+{
+	Super::Serialize(Ar);
+
+	Ar.UsingCustomVersion(FDMXPixelMappingMainStreamObjectVersion::GUID);
+	if (Ar.IsLoading())
+	{
+		if (Ar.CustomVer(FDMXPixelMappingMainStreamObjectVersion::GUID) < FDMXPixelMappingMainStreamObjectVersion::RendererComponentHoldsLayoutRect)
+		{	
+			// Assets created before 5.4 do not store the layout rect, so they cannot follow the texture size by default
+			bChildrenFollowSize = false;
+		}
+	}
+}
+
 void UDMXPixelMappingRendererComponent::PostLoad()
 {
 	Super::PostLoad();
@@ -112,6 +118,7 @@ void UDMXPixelMappingRendererComponent::PostEditChangeChainProperty(FPropertyCha
 	if (PropertyChangedChainEvent.ChangeType != EPropertyChangeType::Interactive)
 	{
 		UpdatePreprocessRenderer();
+		LetChildrenFollowSize();
 	}
 
 	PRAGMA_DISABLE_DEPRECATION_WARNINGS
@@ -236,6 +243,8 @@ void UDMXPixelMappingRendererComponent::Render()
 		if (GetSize() != TextureSize)
 		{
 			SetSize(TextureSize);
+			LetChildrenFollowSize();
+
 			bInvalidatePixelMap = true;
 		}
 	}
@@ -311,7 +320,7 @@ FString UDMXPixelMappingRendererComponent::GetUserName() const
 
 UTexture* UDMXPixelMappingRendererComponent::GetRenderedInputTexture() const
 {
-	return  PreprocessRenderer ? PreprocessRenderer->GetRenderedTexture() : nullptr;
+	return PreprocessRenderer ? PreprocessRenderer->GetRenderedTexture() : nullptr;
 }
 
 void UDMXPixelMappingRendererComponent::OnComponentAddedOrRemoved(UDMXPixelMapping* PixelMapping, UDMXPixelMappingBaseComponent* Component)
@@ -319,6 +328,52 @@ void UDMXPixelMappingRendererComponent::OnComponentAddedOrRemoved(UDMXPixelMappi
 	InvalidatePixelMapRenderer();
 }
 
+void UDMXPixelMappingRendererComponent::LetChildrenFollowSize()
+{
+	if (!bChildrenFollowSize || !PreprocessRenderer)
+	{
+		return;
+	}
+
+	const FVector2D NewSize = PreprocessRenderer->GetResultingSize2D();
+	
+	// Handle the case where the layout rect was never stored (new assets, and assets created before 5.4).
+	// In this case, simply initialize the LayoutRect member.
+	if (LayoutRect == FVector2D::ZeroVector)
+	{
+		LayoutRect = NewSize;
+		return;
+	}
+
+	// Skip unchanged values, or if the current size is zero (no texture).
+	if (NewSize == LayoutRect ||
+		NewSize == FVector2D::ZeroVector)
+	{
+		return;
+	}
+
+	// Scale position and size of all children
+	const FVector2D Scalar = NewSize / LayoutRect;
+
+	constexpr bool bRecursive = true;
+	ForEachChild(
+		[&Scalar](UDMXPixelMappingBaseComponent* Component)
+		{
+			if (UDMXPixelMappingOutputComponent* OutputComponent = Cast<UDMXPixelMappingOutputComponent>(Component))
+			{
+				const FVector2D NewPosition = OutputComponent->GetPosition() * Scalar;
+				OutputComponent->SetPosition(NewPosition);
+
+				const FVector2D NewSize = OutputComponent->GetSize() * Scalar;
+				OutputComponent->SetSize(NewSize);
+			}
+		}, 
+		bRecursive);
+
+	// Remember the new layout rect
+	LayoutRect = NewSize;
+}
+
 UWorld* UDMXPixelMappingRendererComponent::TryGetWorld() const
 {
 	UWorld* World = nullptr;
diff --git a/Engine/Plugins/VirtualProduction/DMX/DMXPixelMapping/Source/DMXPixelMappingRuntime/Private/DMXPixelMappingMainStreamObjectVersion.h b/Engine/Plugins/VirtualProduction/DMX/DMXPixelMapping/Source/DMXPixelMappingRuntime/Private/DMXPixelMappingMainStreamObjectVersion.h
--- a/Engine/Plugins/VirtualProduction/DMX/DMXPixelMapping/Source/DMXPixelMappingRuntime/Private/DMXPixelMappingMainStreamObjectVersion.h
+++ b/Engine/Plugins/VirtualProduction/DMX/DMXPixelMapping/Source/DMXPixelMappingRuntime/Private/DMXPixelMappingMainStreamObjectVersion.h
@@ -25,6 +25,9 @@ struct FDMXPixelMappingMainStreamObjectVersion
 		// Update DMXPixelMappingFixtureGroupItemComponent and DMXPixelMappingMatrixCellComponent to use DMXPixelMappingColorSpace
 		UseDMXPixelMappingColorSpace,
 
+		// The renderer component now holds the layout rect, so its children can be layout out relative to the current texture size
+		RendererComponentHoldsLayoutRect,
+
 		// -----<new versions can be added above this line>-------------------------------------------------
 		VersionPlusOne,
 		LatestVersion = VersionPlusOne - 1
diff --git a/Engine/Plugins/VirtualProduction/DMX/DMXPixelMapping/Source/DMXPixelMappingRuntime/Public/Components/DMXPixelMappingRendererComponent.h b/Engine/Plugins/VirtualProduction/DMX/DMXPixelMapping/Source/DMXPixelMappingRuntime/Public/Components/DMXPixelMappingRendererComponent.h
--- a/Engine/Plugins/VirtualProduction/DMX/DMXPixelMapping/Source/DMXPixelMappingRuntime/Public/Components/DMXPixelMappingRendererComponent.h
+++ b/Engine/Plugins/VirtualProduction/DMX/DMXPixelMapping/Source/DMXPixelMappingRuntime/Public/Components/DMXPixelMappingRendererComponent.h
@@ -38,6 +38,7 @@ public:
 	//~ Begin UObject implementation
 protected:
 	virtual void PostInitProperties() override;
+	virtual void Serialize(FArchive& Ar) override;
 	virtual void PostLoad() override;
 #if WITH_EDITOR
 	virtual void PostEditUndo() override;
@@ -81,7 +82,7 @@ public:
 	/** Returns a copy of the current pixel map render elements */
 	TArray<TSharedRef<UE::DMXPixelMapping::Rendering::FPixelMapRenderElement>> GetPixelMapRenderElements() const;
 
-	/** Type of rendering, Texture, Material, UMG, etc... */
+	/** Selects the input type. It can be a Texture, a Material or a User Widget. */
 	UPROPERTY(EditAnywhere, BlueprintReadOnly, Category = "Render Settings")
 	EDMXPixelMappingRendererType RendererType;
 
@@ -98,9 +99,13 @@ public:
 	TSubclassOf<UUserWidget> InputWidget;
 
 	/** The brightness of the renderer */
-	UPROPERTY(EditAnywhere, BlueprintReadOnly, Category = "Render Settings", Meta = (ClampMin = "0", ClampMax = "1", UIMin = "0", UIMax = "1"))
+	UPROPERTY(EditAnywhere, BlueprintReadWrite, Category = "Render Settings", Meta = (ClampMin = "0", UIMin = "0", UIMax = "1"))
 	float Brightness = 1.f;
 
+	/** If true, children are positioned relative to the size of this renderer */
+	UPROPERTY(EditAnywhere, Category = "Layout")
+	bool bChildrenFollowSize = true;
+
 	/** Layout script for the children of this component (hidden in customizations and displayed in its own panel). */
 	UPROPERTY(EditAnywhere, Instanced, Category = "Layout")
 	TObjectPtr<UDMXPixelMappingLayoutScript> LayoutScript;
@@ -109,6 +114,9 @@ private:
 	/** Called when a component was added to or removed from the pixel mapping */
 	void OnComponentAddedOrRemoved(UDMXPixelMapping* PixelMapping, UDMXPixelMappingBaseComponent* Component);
 
+	/** When bChildrenFollowSize is true, rearranges children relatively to the current size */
+	void LetChildrenFollowSize();
+
 	/** Tries to get any world. If with editor, returns the editor world, in game returns GWorld */
 	UWorld* TryGetWorld() const;
 
@@ -130,6 +138,9 @@ private:
 	UPROPERTY()
 	UDMXPixelMappingPixelMapRenderer* PixelMapRenderer;
 
+	/** The current rect in which children are laid out. Useful to compute UV position and size of children. */
+	UPROPERTY()
+	FVector2D LayoutRect;
 
 	//////////////////////
 	// Deprecated Members
