DMX - Add rotation for pixel mapping components, remove deprecated code that would otherwise need to be updated to use rotations.

diff --git a/Engine/Plugins/VirtualProduction/DMX/DMXPixelMapping/Source/DMXPixelMappingEditor/Private/Customizations/DMXPixelMappingDetailCustomization_FixtureGroup.cpp b/Engine/Plugins/VirtualProduction/DMX/DMXPixelMapping/Source/DMXPixelMappingEditor/Private/Customizations/DMXPixelMappingDetailCustomization_FixtureGroup.cpp
--- a/Engine/Plugins/VirtualProduction/DMX/DMXPixelMapping/Source/DMXPixelMappingEditor/Private/Customizations/DMXPixelMappingDetailCustomization_FixtureGroup.cpp
+++ b/Engine/Plugins/VirtualProduction/DMX/DMXPixelMapping/Source/DMXPixelMappingEditor/Private/Customizations/DMXPixelMappingDetailCustomization_FixtureGroup.cpp
@@ -2,13 +2,9 @@
 
 #include "Customizations/DMXPixelMappingDetailCustomization_FixtureGroup.h"
 
-#include "Algo/Find.h"
 #include "Components/DMXPixelMappingFixtureGroupComponent.h"
 #include "DetailLayoutBuilder.h"
-#include "Editor.h"
 #include "IPropertyUtilities.h"
-#include "Misc/CoreDelegates.h"
-#include "Settings/DMXPixelMappingEditorSettings.h"
 
 
 #define LOCTEXT_NAMESPACE "DMXPixelMappingDetailCustomization_FixtureGroup"
@@ -21,15 +17,6 @@ void FDMXPixelMappingDetailCustomization_FixtureGroup::CustomizeDetails(IDetailL
 	InDetailLayout.HideProperty(GET_MEMBER_NAME_CHECKED(UDMXPixelMappingFixtureGroupComponent, DMXLibrary));
 	InDetailLayout.HideProperty(GET_MEMBER_NAME_CHECKED(UDMXPixelMappingFixtureGroupComponent, LayoutScript));
 
-	// Handle size changes
-	SizeXHandle = InDetailLayout.GetProperty(UDMXPixelMappingOutputComponent::GetSizeXPropertyName(), UDMXPixelMappingOutputComponent::StaticClass());
-	SizeXHandle->SetOnPropertyValuePreChange(FSimpleDelegate::CreateSP(this, &FDMXPixelMappingDetailCustomization_FixtureGroup::OnSizePropertyPreChange));
-	SizeXHandle->SetOnPropertyValueChangedWithData(TDelegate<void(const FPropertyChangedEvent&)>::CreateSP(this, &FDMXPixelMappingDetailCustomization_FixtureGroup::OnSizePropertyChanged));
-
-	SizeYHandle = InDetailLayout.GetProperty(UDMXPixelMappingOutputComponent::GetSizeYPropertyName(), UDMXPixelMappingOutputComponent::StaticClass());
-	SizeYHandle->SetOnPropertyValuePreChange(FSimpleDelegate::CreateSP(this, &FDMXPixelMappingDetailCustomization_FixtureGroup::OnSizePropertyPreChange));
-	SizeYHandle->SetOnPropertyValueChangedWithData(TDelegate<void(const FPropertyChangedEvent&)>::CreateSP(this, &FDMXPixelMappingDetailCustomization_FixtureGroup::OnSizePropertyChanged));
-
 	// Listen to component changes
 	UDMXPixelMappingBaseComponent::GetOnComponentAdded().AddSP(this, &FDMXPixelMappingDetailCustomization_FixtureGroup::OnComponentAddedOrRemoved);
 }
@@ -39,106 +26,4 @@ void FDMXPixelMappingDetailCustomization_FixtureGroup::OnComponentAddedOrRemoved
 	PropertyUtilities->RequestRefresh();
 }
 
-void FDMXPixelMappingDetailCustomization_FixtureGroup::OnSizePropertyPreChange()
-{
-	// Remember the sizes if children scale with parent to create the right transaction for interactive changes
-	const FDMXPixelMappingDesignerSettings& DesignerSettings = GetDefault<UDMXPixelMappingEditorSettings>()->DesignerSettings;
-	if (!DesignerSettings.bScaleChildrenWithParent)
-	{
-		return;
-	}
-
-	const TArray<TWeakObjectPtr<UObject>> SelectedObjects = PropertyUtilities->GetSelectedObjects();
-	TArray<UDMXPixelMappingFixtureGroupComponent*> FixtureGroupComponents;
-	for (TWeakObjectPtr<UObject> Object : SelectedObjects)
-	{
-		if (UDMXPixelMappingFixtureGroupComponent* FixtureGroupComponent = Cast<UDMXPixelMappingFixtureGroupComponent>(Object))
-		{
-			FixtureGroupComponents.Add(FixtureGroupComponent);
-		}
-	}
-	if (FixtureGroupComponents.IsEmpty())
-	{
-		return;
-	}
-
-	PreEditChangeComponentToSizeMap.Reset();
-	for (UDMXPixelMappingFixtureGroupComponent* FixtureGroupComponent : FixtureGroupComponents)
-	{
-		PreEditChangeComponentToSizeMap.Add(FixtureGroupComponent, FixtureGroupComponent->GetSize());
-	}
-}
-
-void FDMXPixelMappingDetailCustomization_FixtureGroup::OnSizePropertyChanged(const FPropertyChangedEvent& PropertyChangedEvent)
-{
-	if(PropertyChangedEvent.ChangeType == EPropertyChangeType::Interactive)
-	{
-		GEditor->GetTimerManager()->SetTimerForNextTick(FTimerDelegate::CreateSP(this, &FDMXPixelMappingDetailCustomization_FixtureGroup::HandleSizePropertyChanged));
-	}
-	else
-	{
-		HandleSizePropertyChanged();
-	}
-}
-
-void FDMXPixelMappingDetailCustomization_FixtureGroup::HandleSizePropertyChanged()
-{
-	// Scale children if desired
-	const FDMXPixelMappingDesignerSettings& DesignerSettings = GetDefault<UDMXPixelMappingEditorSettings>()->DesignerSettings;
-	if (!DesignerSettings.bScaleChildrenWithParent ||
-		PreEditChangeComponentToSizeMap.IsEmpty())
-	{
-		return;
-	}
-
-	const TArray<TWeakObjectPtr<UObject>> SelectedObjects = PropertyUtilities->GetSelectedObjects();
-	TArray<UDMXPixelMappingFixtureGroupComponent*> FixtureGroupComponents;
-	for (TWeakObjectPtr<UObject> Object : SelectedObjects)
-	{
-		if (UDMXPixelMappingFixtureGroupComponent* FixtureGroupComponent = Cast<UDMXPixelMappingFixtureGroupComponent>(Object))
-		{
-			FixtureGroupComponents.Add(FixtureGroupComponent);
-		}
-	}
-	if (FixtureGroupComponents.IsEmpty())
-	{
-		return;
-	}
-
-	for (const TTuple<TWeakObjectPtr<UDMXPixelMappingFixtureGroupComponent>, FVector2D>& PreEditChangeComponentToSizeXPair : PreEditChangeComponentToSizeMap)
-	{
-		UDMXPixelMappingFixtureGroupComponent* const* ComponentPtr = Algo::Find(FixtureGroupComponents, PreEditChangeComponentToSizeXPair.Key.Get());
-		if (!ComponentPtr)
-		{
-			continue;
-		}
-
-		const FVector2D GroupPosition = (*ComponentPtr)->GetPosition();
-		const FVector2D OldSize = PreEditChangeComponentToSizeXPair.Value;
-		const FVector2D NewSize = (*ComponentPtr)->GetSize();
-		if (NewSize == FVector2D::ZeroVector || OldSize == NewSize)
-		{
-			// No division by zero, no unchanged values
-			return;
-		}
-
-		const FVector2D RatioVector = NewSize / OldSize;
-		for (UDMXPixelMappingBaseComponent* BaseChild : (*ComponentPtr)->GetChildren())
-		{
-			if (UDMXPixelMappingOutputComponent* Child = Cast<UDMXPixelMappingOutputComponent>(BaseChild))
-			{
-				Child->Modify();
-
-				// Scale size
-				const FVector2D NewChildSize = Child->GetSize() * RatioVector;
-
-				// Scale position
-				const FVector2D NewChildPosition = Child->GetPosition();
-				const FVector2D NewChildPositionRelative = (NewChildPosition - GroupPosition) * RatioVector;
-				Child->SetPosition(GroupPosition + NewChildPositionRelative);
-			}
-		}
-	}
-}
-
 #undef LOCTEXT_NAMESPACE
diff --git a/Engine/Plugins/VirtualProduction/DMX/DMXPixelMapping/Source/DMXPixelMappingEditor/Private/Customizations/DMXPixelMappingDetailCustomization_FixtureGroup.h b/Engine/Plugins/VirtualProduction/DMX/DMXPixelMapping/Source/DMXPixelMappingEditor/Private/Customizations/DMXPixelMappingDetailCustomization_FixtureGroup.h
--- a/Engine/Plugins/VirtualProduction/DMX/DMXPixelMapping/Source/DMXPixelMappingEditor/Private/Customizations/DMXPixelMappingDetailCustomization_FixtureGroup.h
+++ b/Engine/Plugins/VirtualProduction/DMX/DMXPixelMapping/Source/DMXPixelMappingEditor/Private/Customizations/DMXPixelMappingDetailCustomization_FixtureGroup.h
@@ -3,18 +3,14 @@
 #pragma once
 
 #include "IDetailCustomization.h"
-#include "Library/DMXEntityReference.h"
-
-struct FPropertyChangedEvent;
 
 class FDMXPixelMappingToolkit;
+struct FPropertyChangedEvent;
+class IDetailLayoutBuilder;
+class IPropertyUtilities;
 class UDMXPixelMapping;
 class UDMXPixelMappingBaseComponent;
-class UDMXPixelMappingFixtureGroupComponent;
 
-class IDetailLayoutBuilder;
-class IPropertyHandle;
-class IPropertyUtilities;
 
 
 class FDMXPixelMappingDetailCustomization_FixtureGroup
@@ -38,26 +34,8 @@ private:
 	/** Called when a component was added or removed */
 	void OnComponentAddedOrRemoved(UDMXPixelMapping* PixelMapping, UDMXPixelMappingBaseComponent* Component);
 
-	/** Called before the SizeX property changed */
-	void OnSizePropertyPreChange();
-
-	/** Called when the SizeX property changed */
-	void OnSizePropertyChanged(const FPropertyChangedEvent& PropertyChangedEvent);
-
-	/** Handles the size property changed, useful to call on tick */
-	void HandleSizePropertyChanged();
-
 	/** Weak reference to the DMX editor */
 	TWeakPtr<FDMXPixelMappingToolkit> ToolkitWeakPtr;
 
-	/** SizeX before it got changed */
-	TMap<TWeakObjectPtr<UDMXPixelMappingFixtureGroupComponent>, FVector2D> PreEditChangeComponentToSizeMap;
-
-	/** Property handle for the SizeX property */
-	TSharedPtr<IPropertyHandle> SizeXHandle;
-
-	/** Property handle for the SizeY property */
-	TSharedPtr<IPropertyHandle> SizeYHandle;
-
 	TSharedPtr<IPropertyUtilities> PropertyUtilities;
 };
diff --git a/Engine/Plugins/VirtualProduction/DMX/DMXPixelMapping/Source/DMXPixelMappingEditor/Private/Customizations/DMXPixelMappingDetailCustomization_FixtureGroupItem.cpp b/Engine/Plugins/VirtualProduction/DMX/DMXPixelMapping/Source/DMXPixelMappingEditor/Private/Customizations/DMXPixelMappingDetailCustomization_FixtureGroupItem.cpp
--- a/Engine/Plugins/VirtualProduction/DMX/DMXPixelMapping/Source/DMXPixelMappingEditor/Private/Customizations/DMXPixelMappingDetailCustomization_FixtureGroupItem.cpp
+++ b/Engine/Plugins/VirtualProduction/DMX/DMXPixelMapping/Source/DMXPixelMappingEditor/Private/Customizations/DMXPixelMappingDetailCustomization_FixtureGroupItem.cpp
@@ -18,16 +18,6 @@ TSharedRef<IDetailCustomization> FDMXPixelMappingDetailCustomization_FixtureGrou
 
 void FDMXPixelMappingDetailCustomization_FixtureGroupItem::CustomizeDetails(IDetailLayoutBuilder& InDetailLayout)
 {
-	// Hide absolute postition property handles
-	TSharedPtr<IPropertyHandle> PositionXPropertyHandle = InDetailLayout.GetProperty(UDMXPixelMappingOutputComponent::GetPositionXPropertyName());
-	InDetailLayout.HideProperty(PositionXPropertyHandle);
-	TSharedPtr<IPropertyHandle> PositionYPropertyHandle = InDetailLayout.GetProperty(UDMXPixelMappingOutputComponent::GetPositionYPropertyName());
-	InDetailLayout.HideProperty(PositionYPropertyHandle);
-	TSharedPtr<IPropertyHandle> SizeXPropertyHandle = InDetailLayout.GetProperty(UDMXPixelMappingOutputComponent::GetSizeXPropertyName());
-	InDetailLayout.HideProperty(SizeXPropertyHandle);
-	TSharedPtr<IPropertyHandle> SizeYPropertyHandle = InDetailLayout.GetProperty(UDMXPixelMappingOutputComponent::GetSizeYPropertyName());
-	InDetailLayout.HideProperty(SizeXPropertyHandle);
-
 	// Sort categories
 	InDetailLayout.SortCategories([](const TMap<FName, IDetailCategoryBuilder*>& CategoryMap)
 		{
diff --git a/Engine/Plugins/VirtualProduction/DMX/DMXPixelMapping/Source/DMXPixelMappingEditor/Private/Customizations/DMXPixelMappingDetailCustomization_Renderer.cpp b/Engine/Plugins/VirtualProduction/DMX/DMXPixelMapping/Source/DMXPixelMappingEditor/Private/Customizations/DMXPixelMappingDetailCustomization_Renderer.cpp
--- a/Engine/Plugins/VirtualProduction/DMX/DMXPixelMapping/Source/DMXPixelMappingEditor/Private/Customizations/DMXPixelMappingDetailCustomization_Renderer.cpp
+++ b/Engine/Plugins/VirtualProduction/DMX/DMXPixelMapping/Source/DMXPixelMappingEditor/Private/Customizations/DMXPixelMappingDetailCustomization_Renderer.cpp
@@ -34,16 +34,17 @@ void FDMXPixelMappingDetailCustomization_Renderer::CustomizeDetails(IDetailLayou
 	// Hide Output Component properties
 	const TSharedRef<IPropertyHandle> PositionXHandle = DetailLayout.GetProperty(UDMXPixelMappingOutputComponent::GetPositionXPropertyName(), UDMXPixelMappingOutputComponent::StaticClass());
 	DetailLayout.HideProperty(PositionXHandle);
-
-	const TSharedRef<IPropertyHandle> PositionYHandle = DetailLayout.GetProperty(UDMXPixelMappingOutputComponent::GetPositionYPropertyName(), UDMXPixelMappingOutputComponent::StaticClass());
+	const TSharedRef<IPropertyHandle> PositionYHandle = DetailLayout.GetProperty(UDMXPixelMappingOutputComponent::GetPositionXPropertyName(), UDMXPixelMappingOutputComponent::StaticClass());
 	DetailLayout.HideProperty(PositionYHandle);
 
 	const TSharedRef<IPropertyHandle> SizeXHandle = DetailLayout.GetProperty(UDMXPixelMappingOutputComponent::GetSizeXPropertyName(), UDMXPixelMappingOutputComponent::StaticClass());
 	DetailLayout.HideProperty(SizeXHandle);
-
 	const TSharedRef<IPropertyHandle> SizeYHandle = DetailLayout.GetProperty(UDMXPixelMappingOutputComponent::GetSizeYPropertyName(), UDMXPixelMappingOutputComponent::StaticClass());
 	DetailLayout.HideProperty(SizeYHandle);
 
+	const TSharedRef<IPropertyHandle> RotationHandle = DetailLayout.GetProperty(UDMXPixelMappingOutputComponent::GetRotationPropertyName(), UDMXPixelMappingOutputComponent::StaticClass());
+	DetailLayout.HideProperty(RotationHandle);
+
 	const TSharedRef<IPropertyHandle> IsLockInDesignerHandle = DetailLayout.GetProperty(UDMXPixelMappingOutputComponent::GetLockInDesignerPropertyName(), UDMXPixelMappingOutputComponent::StaticClass());
 	DetailLayout.HideProperty(IsLockInDesignerHandle);
 
diff --git a/Engine/Plugins/VirtualProduction/DMX/DMXPixelMapping/Source/DMXPixelMappingEditor/Private/DMXPixelMappingEditorCommands.cpp b/Engine/Plugins/VirtualProduction/DMX/DMXPixelMapping/Source/DMXPixelMappingEditor/Private/DMXPixelMappingEditorCommands.cpp
--- a/Engine/Plugins/VirtualProduction/DMX/DMXPixelMapping/Source/DMXPixelMappingEditor/Private/DMXPixelMappingEditorCommands.cpp
+++ b/Engine/Plugins/VirtualProduction/DMX/DMXPixelMapping/Source/DMXPixelMappingEditor/Private/DMXPixelMappingEditorCommands.cpp
@@ -20,10 +20,12 @@ FDMXPixelMappingEditorCommands::FDMXPixelMappingEditorCommands()
 void FDMXPixelMappingEditorCommands::RegisterCommands()
 {
 	UI_COMMAND(AddMapping, "Add Source", "Adds a new Source Texure, Material or User Widget to the Pixel Map asset.", EUserInterfaceActionType::Button, FInputChord());
-
 	UI_COMMAND(PlayDMX, "Play DMX", "Starts to send DMX", EUserInterfaceActionType::Button, FInputChord());
 	UI_COMMAND(StopPlayingDMX, "Stop Playing DMX", "Stops playing DMX.", EUserInterfaceActionType::Button, FInputChord());
 
+	UI_COMMAND(ResizeMode, "Resize Mode", "Resizes components when transform handles are being dragged", EUserInterfaceActionType::Button, FInputChord());
+	UI_COMMAND(RotateMode, "Rotate Mode", "Rotates components when transform handles are being dragged", EUserInterfaceActionType::Button, FInputChord());
+
 	UI_COMMAND(ToggleGridSnapping, "Toggle Grid Snapping", "Enables/disables grid snapping", EUserInterfaceActionType::Button, FInputChord());
 
 	UI_COMMAND(SizeComponentToTexture, "Size Component to Texture", "Sizes the selected group to Texture.", EUserInterfaceActionType::Button, FInputChord(EModifierKey::Shift, EKeys::S));
diff --git a/Engine/Plugins/VirtualProduction/DMX/DMXPixelMapping/Source/DMXPixelMappingEditor/Private/DMXPixelMappingEditorCommands.h b/Engine/Plugins/VirtualProduction/DMX/DMXPixelMapping/Source/DMXPixelMappingEditor/Private/DMXPixelMappingEditorCommands.h
--- a/Engine/Plugins/VirtualProduction/DMX/DMXPixelMapping/Source/DMXPixelMappingEditor/Private/DMXPixelMappingEditorCommands.h
+++ b/Engine/Plugins/VirtualProduction/DMX/DMXPixelMapping/Source/DMXPixelMappingEditor/Private/DMXPixelMappingEditorCommands.h
@@ -18,10 +18,12 @@ public:
 	TSharedPtr<FUICommandInfo> PlayDMX;
 	TSharedPtr<FUICommandInfo> StopPlayingDMX;
 
-	// Grid snapping related
+	// Designer related
+	TSharedPtr<FUICommandInfo> ResizeMode;
+	TSharedPtr<FUICommandInfo> RotateMode;
+
 	TSharedPtr<FUICommandInfo> ToggleGridSnapping;
 
-	// Layout related
 	TSharedPtr<FUICommandInfo> SizeComponentToTexture;
 
 	TSharedPtr<FUICommandInfo> ToggleScaleChildrenWithParent;
diff --git a/Engine/Plugins/VirtualProduction/DMX/DMXPixelMapping/Source/DMXPixelMappingEditor/Private/DMXPixelMappingEditorStyle.cpp b/Engine/Plugins/VirtualProduction/DMX/DMXPixelMapping/Source/DMXPixelMappingEditor/Private/DMXPixelMappingEditorStyle.cpp
--- a/Engine/Plugins/VirtualProduction/DMX/DMXPixelMapping/Source/DMXPixelMappingEditor/Private/DMXPixelMappingEditorStyle.cpp
+++ b/Engine/Plugins/VirtualProduction/DMX/DMXPixelMapping/Source/DMXPixelMappingEditor/Private/DMXPixelMappingEditorStyle.cpp
@@ -48,7 +48,7 @@ FDMXPixelMappingEditorStyle::FDMXPixelMappingEditorStyle()
 	// Component border style
 	{
 		FSlateBrush* ComponentBorderBrush = new FSlateBrush();
-		ComponentBorderBrush->Margin = FMargin(1.f);
+		ComponentBorderBrush->Margin = FMargin(2.f);
 		ComponentBorderBrush->DrawAs = ESlateBrushDrawType::Border;
 		ComponentBorderBrush->TintColor = FLinearColor::White;
 		Set("DMXPixelMappingEditor.ComponentBorder", ComponentBorderBrush);
diff --git a/Engine/Plugins/VirtualProduction/DMX/DMXPixelMapping/Source/DMXPixelMappingEditor/Private/DragDrop/DMXPixelMappingDragDropOp.cpp b/Engine/Plugins/VirtualProduction/DMX/DMXPixelMapping/Source/DMXPixelMappingEditor/Private/DragDrop/DMXPixelMappingDragDropOp.cpp
--- a/Engine/Plugins/VirtualProduction/DMX/DMXPixelMapping/Source/DMXPixelMappingEditor/Private/DragDrop/DMXPixelMappingDragDropOp.cpp
+++ b/Engine/Plugins/VirtualProduction/DMX/DMXPixelMapping/Source/DMXPixelMappingEditor/Private/DragDrop/DMXPixelMappingDragDropOp.cpp
@@ -72,61 +72,53 @@ void FDMXPixelMappingDragDropOp::LayoutOutputComponents(const FVector2D& GraphSp
 	{
 		return;
 	}
-	const FVector2D Anchor = FirstComponent->GetPosition();
 
-	// Move all to new position
+	// Compute the translation for the first component
+	const FVector2D OldPositionRotated = FirstComponent->GetPositionRotated();
+	const FVector2D DesiredPositionRotated = FVector2D(GraphSpacePosition - GraphSpaceDragOffset);
+	const FVector2D NewPositionRotated = ComputeGridSnapPosition(DesiredPositionRotated);
+
+	// Move all to the new position
+	const FVector2D Translation = NewPositionRotated - OldPositionRotated;
 	for (const TWeakObjectPtr<UDMXPixelMappingBaseComponent>& Component : DraggedComponents)
 	{
 		if (UDMXPixelMappingOutputComponent* OutputComponent = Cast<UDMXPixelMappingOutputComponent>(Component.Get()))
 		{
-			OutputComponent->PreEditChange(nullptr);
-			
-			constexpr bool bModifyChildrenRecursive = true;
-			Component->ForEachChild([](UDMXPixelMappingBaseComponent* Component)
-				{
-					Component->Modify();
-				}, bModifyChildrenRecursive);
-
 			if (ensureMsgf(OutputComponent->GetClass() != UDMXPixelMappingMatrixComponent::StaticClass(),
-				TEXT("Only matrix components can laid out with FDMXPixelMappingDragDropOp::LayoutOutputComponents. Please use FGroupChildDragDropHelper instead (see FDMXPixelMappingDragDropOp::GetGroupChildDragDropHelper().")))
+				TEXT("Matrix components cannot be laid out with FDMXPixelMappingDragDropOp::LayoutOutputComponents. Please use FGroupChildDragDropHelper instead (see FDMXPixelMappingDragDropOp::GetGroupChildDragDropHelper().")))
 			{
-				const FVector2D AnchorOffset = Anchor - OutputComponent->GetPosition();
-
-				const FVector2D NewPosition = FVector2D(GraphSpacePosition - AnchorOffset - GraphSpaceDragOffset).RoundToVector();
-				OutputComponent->SetPosition(NewPosition);
+				OutputComponent->Modify();
+				OutputComponent->SetPosition(OutputComponent->GetPosition() + Translation);
 			}
 		}
 	}
-	GridSnap();
-
-	for (const TWeakObjectPtr<UDMXPixelMappingBaseComponent>& Component : DraggedComponents)
-	{
-		if (UDMXPixelMappingOutputComponent* OutputComponent = Cast<UDMXPixelMappingOutputComponent>(Component.Get()))
-		{
-			OutputComponent->PostEditChange();
-		}
-	}
 }
 
-void FDMXPixelMappingDragDropOp::GridSnap()
+FVector2D FDMXPixelMappingDragDropOp::ComputeGridSnapPosition(const FVector2D& DesiredPosition) const
 {
 	const UDMXPixelMapping* PixelMapping = WeakToolkit.IsValid() ? WeakToolkit.Pin()->GetDMXPixelMapping() : nullptr;
 	if (!PixelMapping || DraggedComponents.IsEmpty() || PixelMapping->SnapGridColumns == 0 || PixelMapping->SnapGridRows == 0)
 	{
-		return;
+		return DesiredPosition;
 	}
 
-	// Grid snap the first component only, and move the others by same delta.
 	UDMXPixelMappingOutputComponent* FirstComponent = Cast<UDMXPixelMappingOutputComponent>(DraggedComponents[0].Get());
 	if (!FirstComponent)
 	{
-		return;
+		return DesiredPosition;
 	}
 
 	UDMXPixelMappingRendererComponent* RendererOfFirstComponent = FirstComponent->GetRendererComponent();
 	if (!RendererOfFirstComponent)
 	{
-		return;
+		return DesiredPosition;
+	}
+
+	// Grid snap to pixels even when grid snapping is disabled, if the first component is axis aligned
+	const bool bIsAxisAligned = FMath::IsNearlyZero(FMath::Abs(FMath::Fmod(FirstComponent->GetRotation(), 90.f)));
+	if (!PixelMapping->bGridSnappingEnabled && !bIsAxisAligned)
+	{
+		return DesiredPosition;
 	}
 
 	const FVector2D CellSize = [PixelMapping, RendererOfFirstComponent]()
@@ -141,20 +133,12 @@ void FDMXPixelMappingDragDropOp::GridSnap()
 			return FVector2D(1.f, 1.f);
 		}();
 
-	const int32 Column = FirstComponent->GetPosition().X / CellSize.X;
-	const int32 Row = FirstComponent->GetPosition().Y / CellSize.Y;
+	const int32 Column = FMath::RoundHalfToZero(DesiredPosition.X / CellSize.X);
+	const int32 Row = FMath::RoundHalfToZero(DesiredPosition.Y / CellSize.Y);
 
 	const FVector2D GridSnapPosition = FVector2D(Column, Row) * CellSize;
-	const FVector2D DeltaVector = GridSnapPosition - FirstComponent->GetPosition();
 
-	FirstComponent->SetPosition(GridSnapPosition);
-	for (int32 ComponentIndex = 1; ComponentIndex < DraggedComponents.Num(); ComponentIndex++)
-	{
-		if (UDMXPixelMappingOutputComponent* OutputComponent = Cast<UDMXPixelMappingOutputComponent>(DraggedComponents[ComponentIndex].Get()))
-		{
-			OutputComponent->SetPosition(OutputComponent->GetPosition() + DeltaVector);
-		}
-	}
+	return GridSnapPosition;
 }
 
 #undef LOCTEXT_NAMESPACE
diff --git a/Engine/Plugins/VirtualProduction/DMX/DMXPixelMapping/Source/DMXPixelMappingEditor/Private/DragDrop/DMXPixelMappingDragDropOp.h b/Engine/Plugins/VirtualProduction/DMX/DMXPixelMapping/Source/DMXPixelMappingEditor/Private/DragDrop/DMXPixelMappingDragDropOp.h
--- a/Engine/Plugins/VirtualProduction/DMX/DMXPixelMapping/Source/DMXPixelMappingEditor/Private/DragDrop/DMXPixelMappingDragDropOp.h
+++ b/Engine/Plugins/VirtualProduction/DMX/DMXPixelMapping/Source/DMXPixelMappingEditor/Private/DragDrop/DMXPixelMappingDragDropOp.h
@@ -42,8 +42,6 @@ public:
 	/** Returns true if the drag drop op was created with the constructor that provides a template */
 	bool WasCreatedAsTemplate() const { return bWasCreatedAsTemplate; }
 
-	/** Utiltity that applies grid snapping. Note, this grid snaps to pixels if grid snapping is disabled. */
-	void GridSnap();
 
 	/** Can be used to store the parent of the dragged component(s) */
 	TWeakObjectPtr<UDMXPixelMappingBaseComponent> Parent;
@@ -59,6 +57,9 @@ public:
 
 	/** Returns the dragged component references */
 	FORCEINLINE const TArray<TSharedPtr<FDMXPixelMappingComponentTemplate>>& GetTemplates() const { return Templates; }
+	
+	/** Utiltity that returns the grid snapped position from a desired position. */
+	FVector2D ComputeGridSnapPosition(const FVector2D& DesiredPosition) const;
 
 private:
 	/** True if this was created with the New version was used that provides a template */
diff --git a/Engine/Plugins/VirtualProduction/DMX/DMXPixelMapping/Source/DMXPixelMappingEditor/Private/DragDrop/DMXPixelMappingGroupChildDragDropHelper.cpp b/Engine/Plugins/VirtualProduction/DMX/DMXPixelMapping/Source/DMXPixelMappingEditor/Private/DragDrop/DMXPixelMappingGroupChildDragDropHelper.cpp
--- a/Engine/Plugins/VirtualProduction/DMX/DMXPixelMapping/Source/DMXPixelMappingEditor/Private/DragDrop/DMXPixelMappingGroupChildDragDropHelper.cpp
+++ b/Engine/Plugins/VirtualProduction/DMX/DMXPixelMapping/Source/DMXPixelMappingEditor/Private/DragDrop/DMXPixelMappingGroupChildDragDropHelper.cpp
@@ -30,7 +30,7 @@ TSharedPtr<FDMXPixelMappingGroupChildDragDropHelper> FDMXPixelMappingGroupChildD
 		return nullptr;
 	}
 	NewHelper->WeakParentGroupComponent = ParentGroupComponent;
-	NewHelper->ParentPosition = ParentGroupComponent->GetPosition();
+	NewHelper->ParentPosition = ParentGroupComponent->GetPositionRotated();
 	NewHelper->ParentSize = ParentGroupComponent->GetSize();
 
 	for (const TWeakObjectPtr<UDMXPixelMappingBaseComponent>& BaseComponent : DragDropOp->GetDraggedComponents())
@@ -86,7 +86,7 @@ void FDMXPixelMappingGroupChildDragDropHelper::LayoutAligned(const FVector2D& Gr
 	{
 		return;
 	}
-
+	 
 	FVector2D NextPosition = GraphSpacePosition - DragDropOp->GraphSpaceDragOffset;
 	float RowHeight = 0.f;
 
@@ -105,7 +105,7 @@ void FDMXPixelMappingGroupChildDragDropHelper::LayoutAligned(const FVector2D& Gr
 			if (GroupComponent->IsOverPosition(NextPosition) && 
 				GroupComponent->IsOverPosition(NextPosition + ChildComponent->GetSize()))
 			{
-				ChildComponent->SetPosition(NextPosition);
+				ChildComponent->SetPositionRotated(NextPosition);
 
 				RowHeight = FMath::Max(ChildComponent->GetSize().Y, RowHeight);
 				NextPosition = FVector2D(NextPosition.X + ChildComponent->GetSize().X, NextPosition.Y);
@@ -119,22 +119,22 @@ void FDMXPixelMappingGroupChildDragDropHelper::LayoutAligned(const FVector2D& Gr
 				if (GroupComponent->IsOverPosition(NextPositionOnNewRow) &&
 					GroupComponent->IsOverPosition(NextPositionOnNewRow + ChildComponent->GetSize()))
 				{
-					ChildComponent->SetPosition(NewRowPosition);
+					ChildComponent->SetPositionRotated(NewRowPosition);
 
 					NextPosition = FVector2D(NewRowPosition.X + ChildComponent->GetSize().X, NewRowPosition.Y);
 					RowHeight = ChildComponent->GetSize().Y;							
 				}
 				else
 				{
-					ChildComponent->SetPosition(NextPosition);
+					ChildComponent->SetPositionRotated(NextPosition);
 
 					NextPosition = FVector2D(NextPosition.X + ChildComponent->GetSize().X, NextPosition.Y);
 				}
 			}
 		}
 	}
-
-	DragDropOp->GridSnap();
+	const FVector2D GridSnapPosition = DragDropOp->ComputeGridSnapPosition(GroupComponent->GetPositionRotated());
+	GroupComponent->SetPositionRotated(GridSnapPosition);
 
 	for (const TWeakObjectPtr<UDMXPixelMappingOutputComponent>& WeakChildComponent : WeakChildComponents)
 	{
@@ -159,7 +159,7 @@ void FDMXPixelMappingGroupChildDragDropHelper::LayoutUnaligned(const FVector2D&
 		return;
 	}
 
-	const FVector2D Anchor = FirstComponent->GetPosition();
+	const FVector2D Anchor = FirstComponent->GetPositionRotated();
 	for (const TWeakObjectPtr<UDMXPixelMappingOutputComponent>& WeakOutputComponent : WeakChildComponents)
 	{
 		if (UDMXPixelMappingOutputComponent* ChildComponent = WeakOutputComponent.Get())
@@ -172,14 +172,14 @@ void FDMXPixelMappingGroupChildDragDropHelper::LayoutUnaligned(const FVector2D&
 					Component->Modify();
 				}, bModifyChildrenRecursive);
 
-			const FVector2D AnchorOffset = Anchor - ChildComponent->GetPosition();
+			const FVector2D AnchorOffset = Anchor - ChildComponent->GetPositionRotated();
 
 			const FVector2D NewPosition = GraphSpacePosition - AnchorOffset - DragDropOp->GraphSpaceDragOffset;
-			ChildComponent->SetPosition(NewPosition);
+			ChildComponent->SetPositionRotated(NewPosition);
 		}
 	}
-
-	DragDropOp->GridSnap();
+	const FVector2D GridSnapPosition = DragDropOp->ComputeGridSnapPosition(FirstComponent->GetPositionRotated());
+	FirstComponent->SetPositionRotated(GridSnapPosition);
 
 	for (const TWeakObjectPtr<UDMXPixelMappingOutputComponent>& WeakOutputComponent : WeakChildComponents)
 	{
diff --git a/Engine/Plugins/VirtualProduction/DMX/DMXPixelMapping/Source/DMXPixelMappingEditor/Private/Toolkits/DMXPixelMappingToolkit.cpp b/Engine/Plugins/VirtualProduction/DMX/DMXPixelMapping/Source/DMXPixelMappingEditor/Private/Toolkits/DMXPixelMappingToolkit.cpp
--- a/Engine/Plugins/VirtualProduction/DMX/DMXPixelMapping/Source/DMXPixelMappingEditor/Private/Toolkits/DMXPixelMappingToolkit.cpp
+++ b/Engine/Plugins/VirtualProduction/DMX/DMXPixelMapping/Source/DMXPixelMappingEditor/Private/Toolkits/DMXPixelMappingToolkit.cpp
@@ -42,7 +42,6 @@
 const FName FDMXPixelMappingToolkit::LayoutViewTabID(TEXT("DMXPixelMappingEditor_LayoutViewTabID"));
 
 FDMXPixelMappingToolkit::FDMXPixelMappingToolkit()
-	: DMXPixelMapping(nullptr)
 {
 	EditorSettingsDump = TArray<uint8, TFixedAllocator<sizeof(UDMXPixelMappingEditorSettings)>>(reinterpret_cast<const uint8*>(GetDefault<UDMXPixelMappingEditorSettings>()), (int32)sizeof(UDMXPixelMappingEditorSettings));
 }
@@ -57,9 +56,122 @@
 	check(InDMXPixelMapping);
 	InDMXPixelMapping->DestroyInvalidComponents();
 
-	DMXPixelMapping = InDMXPixelMapping;
+	// Make sure we loaded all UObjects
+	InDMXPixelMapping->CreateOrLoadObjects();
 
-	InitializeInternal(Mode, InitToolkitHost, FGuid::NewGuid());
+	// Bind to component changes
+	UDMXPixelMappingBaseComponent::GetOnComponentAdded().AddSP(this, &FDMXPixelMappingToolkit::OnComponentAddedOrRemoved);
+	UDMXPixelMappingBaseComponent::GetOnComponentRemoved().AddSP(this, &FDMXPixelMappingToolkit::OnComponentAddedOrRemoved);
+	UDMXPixelMappingBaseComponent::GetOnComponentRenamed().AddSP(this, &FDMXPixelMappingToolkit::OnComponentRenamed);
+
+	SetupCommands();
+
+	CreateInternalViews();
+
+	TSharedRef<FTabManager::FLayout> StandaloneDefaultLayout = FTabManager::NewLayout("Standalone_PixelMapping_Layout_2.0")
+		->AddArea
+		(
+			FTabManager::NewPrimaryArea()->SetOrientation(Orient_Vertical)
+			->Split
+			(
+				FTabManager::NewSplitter()
+				->SetOrientation(Orient_Horizontal)
+				->Split
+				(
+					FTabManager::NewSplitter()
+					->SetOrientation(Orient_Vertical)
+					->SetSizeCoefficient(.25f)
+					->Split
+					(
+						FTabManager::NewStack()
+						->AddTab(DMXLibraryViewTabID, ETabState::OpenedTab)
+						->SetSizeCoefficient(.382f)
+					)
+					->Split
+					(
+						FTabManager::NewStack()
+						->AddTab(HierarchyViewTabID, ETabState::OpenedTab)
+						->SetSizeCoefficient(.618f)
+					)
+				)
+				->Split
+				(
+					FTabManager::NewSplitter()
+					->SetOrientation(Orient_Vertical)
+					->SetSizeCoefficient(.5f)
+					->Split
+					(
+						FTabManager::NewStack()
+						->AddTab(DesignerViewTabID, ETabState::OpenedTab)
+						->SetSizeCoefficient(.75f)
+					)
+					->Split
+					(
+						FTabManager::NewStack()
+						->AddTab(PreviewViewTabID, ETabState::OpenedTab)
+						->SetSizeCoefficient(.25f)
+					)
+				)
+				->Split
+				(
+					FTabManager::NewSplitter()
+					->SetOrientation(Orient_Vertical)
+					->SetSizeCoefficient(.25f)
+					->Split
+					(
+						FTabManager::NewStack()
+						->AddTab(DetailsViewTabID, ETabState::OpenedTab)
+						->SetSizeCoefficient(.618f)
+					)
+					->Split
+					(
+						FTabManager::NewStack()
+						->AddTab(LayoutViewTabID, ETabState::OpenedTab)
+						->SetSizeCoefficient(.382f)
+					)
+				)
+			)
+		);
+
+	const bool bCreateDefaultStandaloneMenu = true;
+	const bool bCreateDefaultToolbar = true;
+	FAssetEditorToolkit::InitAssetEditor(Mode, InitToolkitHost, FDMXPixelMappingEditorModule::DMXPixelMappingEditorAppIdentifier,
+		StandaloneDefaultLayout, bCreateDefaultStandaloneMenu, bCreateDefaultToolbar, InDMXPixelMapping);
+
+	// Allow extenders to extend the toolbar, then regenerate menus and toolbars.
+	ExtendToolbar();
+	RegenerateMenusAndToolbars();
+
+	// Make an initial selection
+	if (UDMXPixelMappingRootComponent* RootComponent = InDMXPixelMapping->GetRootComponent())
+	{
+		UDMXPixelMappingBaseComponent* const* FirstRendererComponetPtr = Algo::FindByPredicate(InDMXPixelMapping->GetRootComponent()->GetChildren(), [](UDMXPixelMappingBaseComponent* Component)
+			{
+				return Component && Component->GetClass() == UDMXPixelMappingRendererComponent::StaticClass();
+			});
+		if (FirstRendererComponetPtr)
+		{
+			UDMXPixelMappingBaseComponent* const* FirstFixtureGroupComponetPtr = Algo::FindByPredicate((*FirstRendererComponetPtr)->GetChildren(), [](UDMXPixelMappingBaseComponent* Component)
+				{
+					return Component && Component->GetClass() == UDMXPixelMappingFixtureGroupComponent::StaticClass();
+				});
+
+			if (UDMXPixelMappingBaseComponent* const* ComponentToSelectPtr = FirstFixtureGroupComponetPtr ?
+				FirstFixtureGroupComponetPtr :
+				FirstRendererComponetPtr)
+			{
+				const FDMXPixelMappingComponentReference ComponentReference(StaticCastSharedRef<FDMXPixelMappingToolkit>(AsShared()), *ComponentToSelectPtr);
+				SelectComponents(TSet<FDMXPixelMappingComponentReference>({ ComponentReference }));
+			}
+		}
+	}
+
+	// Refresh the hierarchy view, so that it displays the pixelmapping of the now-initialized asset editor.
+	HierarchyView->RequestRefresh();
+
+	// Set the scale children with parent property on the pixel mapping object, so it is accessible the runtime module.
+	const UDMXPixelMappingEditorSettings* EditorSettings = GetDefault<UDMXPixelMappingEditorSettings>();
+	InDMXPixelMapping->bEditorScaleChildrenWithParent = EditorSettings->DesignerSettings.bScaleChildrenWithParent;
 }
 
 void FDMXPixelMappingToolkit::RegisterTabSpawners(const TSharedRef<class FTabManager>& InTabManager)
@@ -128,12 +240,13 @@
 
 void FDMXPixelMappingToolkit::Tick(float DeltaTime)
 {
-	if (!ensure(DMXPixelMapping))
+	UDMXPixelMapping* PixelMapping = GetDMXPixelMapping();
+	if (!PixelMapping)
 	{
 		return;
 	}
 
-	UDMXPixelMappingRootComponent* RootComponent = DMXPixelMapping->RootComponent;
+	UDMXPixelMappingRootComponent* RootComponent = PixelMapping->RootComponent;
 	if (!ensure(RootComponent))
 	{
 		return;
@@ -166,6 +279,9 @@
 		EditorSettings->OnEditorSettingsChanged.Broadcast();
 
 		EditorSettingsDump = TArray<uint8, TFixedAllocator<sizeof(UDMXPixelMappingEditorSettings)>>(reinterpret_cast<const uint8*>(GetDefault<UDMXPixelMappingEditorSettings>()), (int32)sizeof(UDMXPixelMappingEditorSettings));
+	
+		// Set the scale children with parent property on the pixel mapping object, so it is accessible the runtime module.
+		PixelMapping->bEditorScaleChildrenWithParent = EditorSettings->DesignerSettings.bScaleChildrenWithParent;
 	}
 }
 
@@ -174,6 +290,11 @@
 	RETURN_QUICK_DECLARE_CYCLE_STAT(FDMXPixelMappingToolkit, STATGROUP_Tickables);
 }
 
+UDMXPixelMapping* FDMXPixelMappingToolkit::GetDMXPixelMapping() const
+{
+	return HasEditingObject() ? Cast<UDMXPixelMapping>(GetEditingObject()) : nullptr;
+}
+
 FDMXPixelMappingComponentReference FDMXPixelMappingToolkit::GetReferenceFromComponent(UDMXPixelMappingBaseComponent* InComponent)
 {
 	return FDMXPixelMappingComponentReference(SharedThis(this), InComponent);
@@ -263,13 +384,14 @@
 
 void FDMXPixelMappingToolkit::AddRenderer()
 {
-	UDMXPixelMappingRootComponent* RootComponent = DMXPixelMapping ? DMXPixelMapping->GetRootComponent() : nullptr;
+	UDMXPixelMapping* PixelMapping = GetDMXPixelMapping();
+	UDMXPixelMappingRootComponent* RootComponent = PixelMapping ? PixelMapping->GetRootComponent() : nullptr;
 	if (RootComponent)
 	{
 		const FScopedTransaction AddMappingTransaction(LOCTEXT("AddMappingTransaction", "Add Mapping to Pixel Mapping"));
 
 		RootComponent->PreEditChange(UDMXPixelMappingBaseComponent::StaticClass()->FindPropertyByName(GET_MEMBER_NAME_CHECKED(UDMXPixelMappingBaseComponent, Children)));
-		UDMXPixelMappingRendererComponent* NewRendererComponent = FDMXPixelMappingEditorUtils::AddRenderer(DMXPixelMapping);
+		UDMXPixelMappingRendererComponent* NewRendererComponent = FDMXPixelMappingEditorUtils::AddRenderer(PixelMapping);
 		RootComponent->PostEditChange();
 
 		SetActiveRenderComponent(NewRendererComponent);
@@ -302,9 +424,10 @@
 
 void FDMXPixelMappingToolkit::SaveThumbnailImage()
 {
-	if (DMXPixelMapping && ActiveRendererComponent.IsValid())
+	UDMXPixelMapping* PixelMapping = GetDMXPixelMapping();
+	if (PixelMapping && ActiveRendererComponent.IsValid())
 	{
-		DMXPixelMapping->ThumbnailImage = ActiveRendererComponent->GetRenderedInputTexture();
+		PixelMapping->ThumbnailImage = ActiveRendererComponent->GetRenderedInputTexture();
 	}
 }
 
@@ -327,6 +450,14 @@
 					NewComponent->Modify();
 
 					Target->AddChild(NewComponent);
+
+					// Output components need to adopt the initial rotation from their parent if possible.
+					UDMXPixelMappingOutputComponent* NewOutputComponent = Cast<UDMXPixelMappingOutputComponent>(NewComponent);
+					UDMXPixelMappingOutputComponent* ParentOutputComponent = NewOutputComponent ? Cast<UDMXPixelMappingOutputComponent>(NewOutputComponent->GetParent()) : nullptr;
+					if (NewOutputComponent && ParentOutputComponent)
+					{
+						NewOutputComponent->SetRotation(ParentOutputComponent->GetRotation());
+					}
 				}
 			}
 		}
@@ -452,42 +583,15 @@
 		SizeComponentToTextureTransaction = MakeShared<FScopedTransaction>(LOCTEXT("SizeComponentToTextureTransaction", "Size Component to Texture"));
 	}
 
-	const FDMXPixelMappingDesignerSettings& DesignerSettings = GetDefault<UDMXPixelMappingEditorSettings>()->DesignerSettings;
-	if (DesignerSettings.bScaleChildrenWithParent)
-	{
-		// Scale children to parent
-		const FVector2D RatioVector = TextureSize / Component->GetSize();
-		for (UDMXPixelMappingBaseComponent* BaseChild : Component->GetChildren())
-		{
-			if (UDMXPixelMappingOutputComponent* Child = Cast<UDMXPixelMappingOutputComponent>(BaseChild))
-			{
-				Child->PreEditChange(UDMXPixelMappingOutputComponent::StaticClass()->FindPropertyByName(UDMXPixelMappingOutputComponent::GetPositionXPropertyName()));
-				Child->PreEditChange(UDMXPixelMappingOutputComponent::StaticClass()->FindPropertyByName(UDMXPixelMappingOutputComponent::GetPositionYPropertyName()));
-				Child->PreEditChange(UDMXPixelMappingOutputComponent::StaticClass()->FindPropertyByName(UDMXPixelMappingOutputComponent::GetSizeXPropertyName()));
-				Child->PreEditChange(UDMXPixelMappingOutputComponent::StaticClass()->FindPropertyByName(UDMXPixelMappingOutputComponent::GetSizeYPropertyName()));
-
-				// Scale size
-				FVector2D NewSize = Child->GetSize() * RatioVector;
-
-				// Scale position (new position is zero vector)
-				const FVector2D ChildPosition = Child->GetPosition();
-				const FVector2D NewPositionRelative = (ChildPosition - Component->GetPosition()) * RatioVector;
-				Child->SetPosition(Component->GetPosition() + NewPositionRelative);
-
-				Child->PostEditChange();
-			}
-		}
-	}
-
-	Component->PreEditChange(UDMXPixelMappingOutputComponent::StaticClass()->FindPropertyByName(UDMXPixelMappingOutputComponent::GetPositionXPropertyName()));
-	Component->PreEditChange(UDMXPixelMappingOutputComponent::StaticClass()->FindPropertyByName(UDMXPixelMappingOutputComponent::GetPositionYPropertyName()));
-	Component->PreEditChange(UDMXPixelMappingOutputComponent::StaticClass()->FindPropertyByName(UDMXPixelMappingOutputComponent::GetSizeXPropertyName()));
-	Component->PreEditChange(UDMXPixelMappingOutputComponent::StaticClass()->FindPropertyByName(UDMXPixelMappingOutputComponent::GetSizeYPropertyName()));
-
+	Component->Modify();
+	Component->SetRotation(0.0);
 	Component->SetPosition(FVector2D::ZeroVector);
 	Component->SetSize(TextureSize);
-
-	Component->PostEditChange();
 }
 
+void FDMXPixelMappingToolkit::SetTransformHandleMode(EDMXPixelMappingTransformHandleMode NewTransformHandleMode)
+{
+	TransformHandleMode = NewTransformHandleMode;
+}
+
 void FDMXPixelMappingToolkit::ToggleGridSnapping()
@@ -503,7 +607,8 @@
 
 void FDMXPixelMappingToolkit::PostUndo(bool bSuccess)
 {
-	UDMXPixelMappingRootComponent* RootComponent = DMXPixelMapping ? DMXPixelMapping->GetRootComponent() : nullptr;
+	UDMXPixelMapping* PixelMapping = GetDMXPixelMapping();
+	UDMXPixelMappingRootComponent* RootComponent = PixelMapping ? PixelMapping->GetRootComponent() : nullptr;
 	if (!RootComponent)
 	{
 		return;
@@ -539,128 +644,6 @@
 }
 
 
-void FDMXPixelMappingToolkit::InitializeInternal(const EToolkitMode::Type Mode, const TSharedPtr<class IToolkitHost>& InitToolkitHost, const FGuid& MessageLogGuid)
-{
-	if (!DMXPixelMapping)
-	{
-		return;
-	}
-
-	// Make sure we loaded all UObjects
-	DMXPixelMapping->CreateOrLoadObjects();
-
-	// Bind to component changes
-	UDMXPixelMappingBaseComponent::GetOnComponentAdded().AddSP(this, &FDMXPixelMappingToolkit::OnComponentAddedOrRemoved);
-	UDMXPixelMappingBaseComponent::GetOnComponentRemoved().AddSP(this, &FDMXPixelMappingToolkit::OnComponentAddedOrRemoved);
-	UDMXPixelMappingBaseComponent::GetOnComponentRenamed().AddSP(this, &FDMXPixelMappingToolkit::OnComponentRenamed);
-
-	// Create commands
-	DesignerCommandList = MakeShareable(new FUICommandList);
-	DesignerCommandList->MapAction(FGenericCommands::Get().Delete,
-		FExecuteAction::CreateSP(this, &FDMXPixelMappingToolkit::DeleteSelectedComponents)
-	);
-
-	CreateInternalViews();
-
-	TSharedRef<FTabManager::FLayout> StandaloneDefaultLayout = FTabManager::NewLayout("Standalone_PixelMapping_Layout_2.0")
-		->AddArea
-		(
-			FTabManager::NewPrimaryArea()->SetOrientation(Orient_Vertical)
-			->Split
-			(
-				FTabManager::NewSplitter()
-				->SetOrientation(Orient_Horizontal)
-				->Split
-				(
-					FTabManager::NewSplitter()
-					->SetOrientation(Orient_Vertical)
-					->SetSizeCoefficient(.25f)
-					->Split
-					(
-						FTabManager::NewStack()
-						->AddTab(DMXLibraryViewTabID, ETabState::OpenedTab)
-						->SetSizeCoefficient(.382f)
-					)
-					->Split
-					(
-						FTabManager::NewStack()
-						->AddTab(HierarchyViewTabID, ETabState::OpenedTab)
-						->SetSizeCoefficient(.618f)
-					)
-				)
-				->Split
-				(
-					FTabManager::NewSplitter()
-					->SetOrientation(Orient_Vertical)
-					->SetSizeCoefficient(.5f)
-					->Split
-					(
-						FTabManager::NewStack()
-						->AddTab(DesignerViewTabID, ETabState::OpenedTab)
-						->SetSizeCoefficient(.75f)
-					)
-					->Split
-					(
-						FTabManager::NewStack()
-						->AddTab(PreviewViewTabID, ETabState::OpenedTab)
-						->SetSizeCoefficient(.25f)
-					)
-				)
-				->Split
-				(
-					FTabManager::NewSplitter()
-					->SetOrientation(Orient_Vertical)
-					->SetSizeCoefficient(.25f)
-					->Split
-					(
-						FTabManager::NewStack()
-						->AddTab(DetailsViewTabID, ETabState::OpenedTab)
-						->SetSizeCoefficient(.618f)
-					)
-					->Split
-					(
-						FTabManager::NewStack()
-						->AddTab(LayoutViewTabID, ETabState::OpenedTab)
-						->SetSizeCoefficient(.382f)
-					)
-				)
-			)
-		);
-
-	const bool bCreateDefaultStandaloneMenu = true;
-	const bool bCreateDefaultToolbar = true;
-	FAssetEditorToolkit::InitAssetEditor(Mode, InitToolkitHost, FDMXPixelMappingEditorModule::DMXPixelMappingEditorAppIdentifier,
-		StandaloneDefaultLayout, bCreateDefaultStandaloneMenu, bCreateDefaultToolbar, DMXPixelMapping);
-	
-	SetupCommands();
-	ExtendToolbar();
-	RegenerateMenusAndToolbars();
-	
-	// Make an initial selection
-	if (UDMXPixelMappingRootComponent* RootComponent = DMXPixelMapping->GetRootComponent())
-	{
-		UDMXPixelMappingBaseComponent* const* FirstRendererComponetPtr = Algo::FindByPredicate(DMXPixelMapping->GetRootComponent()->GetChildren(), [](UDMXPixelMappingBaseComponent* Component)
-			{
-				return Component && Component->GetClass() == UDMXPixelMappingRendererComponent::StaticClass();
-			});
-		if (FirstRendererComponetPtr)
-		{
-			UDMXPixelMappingBaseComponent* const* FirstFixtureGroupComponetPtr = Algo::FindByPredicate((*FirstRendererComponetPtr)->GetChildren(), [](UDMXPixelMappingBaseComponent* Component)
-				{
-					return Component && Component->GetClass() == UDMXPixelMappingFixtureGroupComponent::StaticClass();
-				});
-
-			if (UDMXPixelMappingBaseComponent* const* ComponentToSelectPtr = FirstFixtureGroupComponetPtr ? 
-				FirstFixtureGroupComponetPtr : 
-				FirstRendererComponetPtr)
-			{
-				const FDMXPixelMappingComponentReference ComponentReference(StaticCastSharedRef<FDMXPixelMappingToolkit>(AsShared()), *ComponentToSelectPtr);
-				SelectComponents(TSet<FDMXPixelMappingComponentReference>({ ComponentReference }));
-			}
-		}
-	}
-}
-
 TSharedRef<SDockTab> FDMXPixelMappingToolkit::SpawnTab_DMXLibraryView(const FSpawnTabArgs& Args)
 {
 	check(Args.GetTabId() == DMXLibraryViewTabID);
@@ -751,19 +734,20 @@
 
 void FDMXPixelMappingToolkit::RenameComponent(const FName& CurrentObjectName, const FString& DesiredObjectName) const
 {
-	if (!DMXPixelMapping)
+	const UDMXPixelMapping* PixelMapping = GetDMXPixelMapping();
+	if (!PixelMapping)
 	{
 		return;
 	}
 
-	UDMXPixelMappingBaseComponent* ComponentToRename = DMXPixelMapping->FindComponent(CurrentObjectName);
+	UDMXPixelMappingBaseComponent* ComponentToRename = PixelMapping->FindComponent(CurrentObjectName);
 	if (!ensureMsgf(ComponentToRename, TEXT("Cannot find component '%s' to rename."), *CurrentObjectName.ToString()))
 	{
 		return;
 	}
 
 	const FName DesiredDisplayName = MakeObjectNameFromDisplayLabel(DesiredObjectName, ComponentToRename->GetFName());
-	UDMXPixelMappingBaseComponent* ExistingComponent = DMXPixelMapping->FindComponent(DesiredDisplayName);
+	UDMXPixelMappingBaseComponent* ExistingComponent = PixelMapping->FindComponent(DesiredDisplayName);
 
 	const FName UniqueName = ExistingComponent ?
 		MakeUniqueObjectName(ComponentToRename->GetOuter(), ComponentToRename->GetClass(), DesiredDisplayName) :
@@ -853,6 +837,13 @@
 
 void FDMXPixelMappingToolkit::SetupCommands()
 {
+	// Create a command list for the designer view specifically
+	DesignerCommandList = MakeShareable(new FUICommandList);
+	DesignerCommandList->MapAction(FGenericCommands::Get().Delete,
+		FExecuteAction::CreateSP(this, &FDMXPixelMappingToolkit::DeleteSelectedComponents)
+	);
+
+	// Init the command list for this toolkit 
 	GetToolkitCommands()->MapAction(
 		FDMXPixelMappingEditorCommands::Get().AddMapping,
 		FExecuteAction::CreateSP(this, &FDMXPixelMappingToolkit::AddRenderer)
@@ -887,5 +878,19 @@
 	);
 
 	GetToolkitCommands()->MapAction(
+		FDMXPixelMappingEditorCommands::Get().ResizeMode,
+		FExecuteAction::CreateSP(this, &FDMXPixelMappingToolkit::SetTransformHandleMode, EDMXPixelMappingTransformHandleMode::Resize),
+		FCanExecuteAction(),
+		FGetActionCheckState::CreateSP(this, &FDMXPixelMappingToolkit::GetTransformHandleModeCheckboxState, EDMXPixelMappingTransformHandleMode::Resize)
+	);
+
+	GetToolkitCommands()->MapAction(
+		FDMXPixelMappingEditorCommands::Get().RotateMode,
+		FExecuteAction::CreateSP(this, &FDMXPixelMappingToolkit::SetTransformHandleMode, EDMXPixelMappingTransformHandleMode::Rotate),
+		FCanExecuteAction(),
+		FGetActionCheckState::CreateSP(this, &FDMXPixelMappingToolkit::GetTransformHandleModeCheckboxState, EDMXPixelMappingTransformHandleMode::Rotate)
+	);
+
+	GetToolkitCommands()->MapAction(
 		FDMXPixelMappingEditorCommands::Get().ToggleGridSnapping,
 		FExecuteAction::CreateSP(this, &FDMXPixelMappingToolkit::ToggleGridSnapping)
@@ -924,4 +929,9 @@
 	AddToolbarExtender(DMXPixelMappingEditorModule.GetToolBarExtensibilityManager()->GetAllExtenders(GetToolkitCommands(), GetEditingObjects()));
 }
 
+ECheckBoxState FDMXPixelMappingToolkit::GetTransformHandleModeCheckboxState(EDMXPixelMappingTransformHandleMode CompareTransformHandleMode) const
+{
+	return CompareTransformHandleMode == TransformHandleMode ? ECheckBoxState::Checked : ECheckBoxState::Unchecked;
+}
+
 #undef LOCTEXT_NAMESPACE
diff --git a/Engine/Plugins/VirtualProduction/DMX/DMXPixelMapping/Source/DMXPixelMappingEditor/Private/Toolkits/DMXPixelMappingToolkit.h b/Engine/Plugins/VirtualProduction/DMX/DMXPixelMapping/Source/DMXPixelMappingEditor/Private/Toolkits/DMXPixelMappingToolkit.h
--- a/Engine/Plugins/VirtualProduction/DMX/DMXPixelMapping/Source/DMXPixelMappingEditor/Private/Toolkits/DMXPixelMappingToolkit.h
+++ b/Engine/Plugins/VirtualProduction/DMX/DMXPixelMapping/Source/DMXPixelMappingEditor/Private/Toolkits/DMXPixelMappingToolkit.h
@@ -7,6 +7,7 @@
 #include "Settings/DMXPixelMappingEditorSettings.h"
 #include "TickableEditorObject.h"
 #include "Toolkits/AssetEditorToolkit.h"
+#include "UObject/GCObject.h"
 #include "Widgets/Views/SHeaderRow.h"
 
 class FDMXPixelMappingComponentTemplate;
@@ -25,6 +26,16 @@ class UDMXPixelMappingBaseComponent;
 class UDMXPixelMappingMatrixComponent;
 class UDMXPixelMappingOutputComponent;
 class UDMXPixelMappingRendererComponent;
+enum class ECheckBoxState : uint8;
+
+namespace UE::DMX
+{
+	enum class EDMXPixelMappingTransformHandleMode : uint8
+	{
+		Resize,
+		Rotate
+	};
+}
 
 
 /**
@@ -35,6 +46,8 @@ class FDMXPixelMappingToolkit
 	, public FTickableEditorObject
 	, public FSelfRegisteringEditorUndoClient
 {
+	using EDMXPixelMappingTransformHandleMode = UE::DMX::EDMXPixelMappingTransformHandleMode;
+
 public:
 	DECLARE_MULTICAST_DELEGATE(FOnSelectedComponentsChangedDelegate)
 	FOnSelectedComponentsChangedDelegate& GetOnSelectedComponentsChangedDelegate() { return OnSelectedComponentsChangedDelegate; }
@@ -75,8 +88,8 @@ public:
 	virtual bool IsTickable() const override { return true; }
 	virtual TStatId GetStatId() const override;
 	//~ End FTickableEditorObject Interface
-
-	UDMXPixelMapping* GetDMXPixelMapping() const { return DMXPixelMapping; }
+		
+	UDMXPixelMapping* GetDMXPixelMapping() const;
 
 	FDMXPixelMappingComponentReference GetReferenceFromComponent(UDMXPixelMappingBaseComponent* InComponent);
 
@@ -141,6 +154,12 @@ public:
 	/** Toggles between grid snapping enabled and disabled */
 	void ToggleGridSnapping();
 
+	/** Sets how transform handles operate */
+	void SetTransformHandleMode(EDMXPixelMappingTransformHandleMode NewTransformHandleMode);
+
+	/** Returns the current transform handle mode */
+	EDMXPixelMappingTransformHandleMode GetTransformHandleMode() const { return TransformHandleMode; }
+
 private:
 	//~ Begin FSelfRegisteringEditorUndoClient interface
 	virtual void PostUndo(bool bSuccess) override;
@@ -163,8 +182,6 @@ private:
 	/** Saves a thumbnail image for the pixel mapping asset */
 	void SaveThumbnailImage();
 
-	void InitializeInternal(const EToolkitMode::Type Mode, const TSharedPtr<class IToolkitHost>& InitToolkitHost, const FGuid& MessageLogGuid);
-
 	/** Spawns the DMX Library View */
 	TSharedRef<SDockTab> SpawnTab_DMXLibraryView(const FSpawnTabArgs& Args);
 
@@ -189,7 +206,8 @@ private:
 
 	void CreateInternalViews();
 
-	UDMXPixelMapping* DMXPixelMapping;
+	/** Returns the checkbox state for a transform handle mode, checked if the mode equals the current mode. */
+	ECheckBoxState GetTransformHandleModeCheckboxState(EDMXPixelMappingTransformHandleMode CompareTransformHandleMode) const;
 
 	/** List of open tool panels; used to ensure only one exists at any one time */
 	TMap<FName, TWeakPtr<SDockableTab>> SpawnedToolPanels;
@@ -240,6 +258,9 @@ private:
 	/** True while removing components (to avoid needlessly updating blueprint nodes on each component removed via our own methods) */
 	bool bRemovingComponents = false;
 
+	/** The current transform handle mode */
+	EDMXPixelMappingTransformHandleMode TransformHandleMode = EDMXPixelMappingTransformHandleMode::Resize;
+
 public:
 	/** Name of the DMX Library View Tab */
 	static const FName DMXLibraryViewTabID;
diff --git a/Engine/Plugins/VirtualProduction/DMX/DMXPixelMapping/Source/DMXPixelMappingEditor/Private/ViewModels/DMXPixelMappingDMXLibraryViewModel.cpp b/Engine/Plugins/VirtualProduction/DMX/DMXPixelMapping/Source/DMXPixelMappingEditor/Private/ViewModels/DMXPixelMappingDMXLibraryViewModel.cpp
--- a/Engine/Plugins/VirtualProduction/DMX/DMXPixelMapping/Source/DMXPixelMappingEditor/Private/ViewModels/DMXPixelMappingDMXLibraryViewModel.cpp
+++ b/Engine/Plugins/VirtualProduction/DMX/DMXPixelMapping/Source/DMXPixelMappingEditor/Private/ViewModels/DMXPixelMappingDMXLibraryViewModel.cpp
@@ -297,6 +297,9 @@ void UDMXPixelMappingDMXLibraryViewModel::LayoutEvenOverParent(const TArray<UDMX
 		return;
 	}
 
+	const double RestoreRotation = FixtureGroupComponent->GetRotation();
+	FixtureGroupComponent->SetRotation(0.0);
+
 	const int32 Columns = FMath::RoundFromZero(FMath::Sqrt((float)Components.Num()));
 	const int32 Rows = FMath::RoundFromZero((float)Components.Num() / Columns);
 	const FVector2D Size = FVector2D(FixtureGroupComponent->GetSize().X / Columns, FixtureGroupComponent->GetSize().Y / Rows);
@@ -319,6 +322,8 @@ void UDMXPixelMappingDMXLibraryViewModel::LayoutEvenOverParent(const TArray<UDMX
 			OutputComponent->SetSize(Size);
 		}
 	}
+
+	FixtureGroupComponent->SetRotation(RestoreRotation);
 }
 
 void UDMXPixelMappingDMXLibraryViewModel::LayoutAfterLastPatch(const TArray<UDMXPixelMappingBaseComponent*> Components)
@@ -335,6 +340,9 @@ void UDMXPixelMappingDMXLibraryViewModel::LayoutAfterLastPatch(const TArray<UDMX
 		return;
 	}
 
+	const double RestoreRotation = FixtureGroupComponent->GetRotation();
+	FixtureGroupComponent->SetRotation(0.0);
+
 	// Find other components
 	TArray<UDMXPixelMappingBaseComponent*> OtherComponentsInGroup;
 	for (UDMXPixelMappingBaseComponent* Child : FixtureGroupComponent->GetChildren())
@@ -397,6 +405,8 @@ void UDMXPixelMappingDMXLibraryViewModel::LayoutAfterLastPatch(const TArray<UDMX
 			}
 		}
 	}
+
+	FixtureGroupComponent->SetRotation(RestoreRotation);
 }
 
 void UDMXPixelMappingDMXLibraryViewModel::SelectFixtureGroupComponent(UDMXPixelMappingFixtureGroupComponent* FixtureGroupComponent)
diff --git a/Engine/Plugins/VirtualProduction/DMX/DMXPixelMapping/Source/DMXPixelMappingEditor/Private/ViewModels/DMXPixelMappingOutputComponentModel.cpp b/Engine/Plugins/VirtualProduction/DMX/DMXPixelMapping/Source/DMXPixelMappingEditor/Private/ViewModels/DMXPixelMappingOutputComponentModel.cpp
--- a/Engine/Plugins/VirtualProduction/DMX/DMXPixelMapping/Source/DMXPixelMappingEditor/Private/ViewModels/DMXPixelMappingOutputComponentModel.cpp
+++ b/Engine/Plugins/VirtualProduction/DMX/DMXPixelMapping/Source/DMXPixelMappingEditor/Private/ViewModels/DMXPixelMappingOutputComponentModel.cpp
@@ -7,6 +7,7 @@
 #include "Components/DMXPixelMappingFixtureGroupItemComponent.h"
 #include "Components/DMXPixelMappingMatrixCellComponent.h"
 #include "Components/DMXPixelMappingMatrixComponent.h"
+#include "Components/DMXPixelMappingOutputComponent.h"
 #include "Components/DMXPixelMappingScreenComponent.h"
 #include "Library/DMXEntityFixturePatch.h"
 #include "Library/DMXLibrary.h"
@@ -61,6 +62,16 @@ FVector2D FDMXPixelMappingOutputComponentModel::GetSize() const
 	return FVector2D::ZeroVector;
 }
 
+FQuat2D FDMXPixelMappingOutputComponentModel::GetQuaternion() const
+{
+	if (UDMXPixelMappingOutputComponent* OutputComponent = WeakOutputComponent.Get())
+	{
+		const FQuat2D Quat(FMath::DegreesToRadians(OutputComponent->GetRotation()));
+		return Quat;
+	}
+	return FQuat2D();
+}
+
 FText FDMXPixelMappingOutputComponentModel::GetName() const
 {
 	if (UDMXPixelMappingOutputComponent* OutputComponent = Cast<UDMXPixelMappingOutputComponent>(WeakOutputComponent.Get()))
@@ -311,7 +322,7 @@ FLinearColor FDMXPixelMappingScreenComponentModel::GetColor() const
 		}
 		else
 		{
-			return ScreenComponent->GetEditorColor().CopyWithNewOpacity(0.6f);
+			return ScreenComponent->GetEditorColor().CopyWithNewOpacity(0.8f);
 		}
 	}
 
diff --git a/Engine/Plugins/VirtualProduction/DMX/DMXPixelMapping/Source/DMXPixelMappingEditor/Private/ViewModels/DMXPixelMappingOutputComponentModel.h b/Engine/Plugins/VirtualProduction/DMX/DMXPixelMapping/Source/DMXPixelMappingEditor/Private/ViewModels/DMXPixelMappingOutputComponentModel.h
--- a/Engine/Plugins/VirtualProduction/DMX/DMXPixelMapping/Source/DMXPixelMappingEditor/Private/ViewModels/DMXPixelMappingOutputComponentModel.h
+++ b/Engine/Plugins/VirtualProduction/DMX/DMXPixelMapping/Source/DMXPixelMappingEditor/Private/ViewModels/DMXPixelMappingOutputComponentModel.h
@@ -2,7 +2,7 @@
 
 #pragma once
 
-
+#include "Math/TransformCalculus2D.h"
 #include "Math/Vector2D.h"
 #include "UObject/WeakObjectPtr.h"
 
@@ -36,6 +36,9 @@ public:
 	/** Returns the size of the component */
 	FVector2D GetSize() const;
 
+	/** Returns the rotation, in degrees */
+	FQuat2D GetQuaternion() const;
+
 	/** Returns the Name of the component */
 	FText GetName() const;
 
diff --git a/Engine/Plugins/VirtualProduction/DMX/DMXPixelMapping/Source/DMXPixelMappingEditor/Private/Views/SDMXPixelMappingDesignerView.cpp b/Engine/Plugins/VirtualProduction/DMX/DMXPixelMapping/Source/DMXPixelMappingEditor/Private/Views/SDMXPixelMappingDesignerView.cpp
--- a/Engine/Plugins/VirtualProduction/DMX/DMXPixelMapping/Source/DMXPixelMappingEditor/Private/Views/SDMXPixelMappingDesignerView.cpp
+++ b/Engine/Plugins/VirtualProduction/DMX/DMXPixelMapping/Source/DMXPixelMappingEditor/Private/Views/SDMXPixelMappingDesignerView.cpp
@@ -526,7 +526,7 @@ FReply SDMXPixelMappingDesignerView::OnDragDetected(const FGeometry& MyGeometry,
 		// We rely on the first being the clicked component. There's an according comment in the DragDropOp and the detected drag is raised here.
 		UDMXPixelMappingOutputComponent* ClickedComponent = CastChecked<UDMXPixelMappingOutputComponent>(DraggedComponents[0]);
 
-		const FVector2D GraphSpaceDragOffset = DragAnchor - ClickedComponent->GetPosition();
+		const FVector2D GraphSpaceDragOffset = DragAnchor - ClickedComponent->GetPositionRotated();
 
 		TSharedRef<FDMXPixelMappingDragDropOp> DragDropOp = FDMXPixelMappingDragDropOp::New(Toolkit, GraphSpaceDragOffset, DraggedComponents);
 		DragDropOp->SetDecoratorVisibility(false);
@@ -564,30 +564,26 @@ void SDMXPixelMappingDesignerView::OnDragLeave(const FDragDropEvent& DragDropEve
 	SDMXPixelMappingSurface::OnDragLeave(DragDropEvent);
 
 	const TSharedPtr<FDMXPixelMappingDragDropOp> DragDropOp = DragDropEvent.GetOperationAs<FDMXPixelMappingDragDropOp>();
-	if (DragDropOp.IsValid())
+	if (DragDropOp.IsValid() && WeakToolkit.IsValid())
 	{
-		if (TSharedPtr<FDMXPixelMappingToolkit> ToolkitPtr = WeakToolkit.Pin())
-		{
-			FScopedRestoreSelection(ToolkitPtr.ToSharedRef(), StaticCastSharedRef<SDMXPixelMappingDesignerView>(AsShared()));
+		TSharedRef<FDMXPixelMappingToolkit> Toolkit = WeakToolkit.Pin().ToSharedRef();
 
-			// If the drag drop op was dragged in from details or palette, remove the components
-			if (DragDropOp->WasCreatedAsTemplate())
+		// If the drag drop op was dragged in from details or palette, remove the components
+		if (DragDropOp->WasCreatedAsTemplate())
+		{
+			TSet<FDMXPixelMappingComponentReference> Parents;
+			for (const TWeakObjectPtr<UDMXPixelMappingBaseComponent>& Component : DragDropOp->GetDraggedComponents())
 			{
-				TSet<FDMXPixelMappingComponentReference> Parents;
-				for (const TWeakObjectPtr<UDMXPixelMappingBaseComponent>& Component : DragDropOp->GetDraggedComponents())
+				if (Component.IsValid() && Component->GetParent())
 				{
-					if (Component.IsValid() && Component->GetParent())
-					{
-						Parents.Add(ToolkitPtr->GetReferenceFromComponent(Component->GetParent()));
-
-						Component->GetParent()->RemoveChild(Component.Get());
-					}
+					Parents.Add(Toolkit->GetReferenceFromComponent(Component->GetParent()));
+					Component->GetParent()->RemoveChild(Component.Get());
 				}
-
-				// Select parents instead
-				ToolkitPtr->SelectComponents(Parents);
 			}
 
+			// Select parents instead
+			Toolkit->SelectComponents(Parents);
+
 			RebuildDesigner();
 		}
 	}
@@ -766,22 +762,19 @@ FVector2D SDMXPixelMappingDesignerView::GetExtensionPosition(TSharedPtr<SDMXPixe
 {
 	const FDMXPixelMappingComponentReference& SelectedComponent = GetSelectedComponent();
 
-	if (SelectedComponent.IsValid())
+	if (UDMXPixelMappingOutputComponent* OutputComponent = Cast<UDMXPixelMappingOutputComponent>(SelectedComponent.GetComponent()))
 	{
 		FGeometry SelectedComponentGeometry;
-		FGeometry SelectedComponentParentGeometry;
 
-		if (GetComponentGeometry(SelectedComponent, SelectedComponentGeometry))
+		if (GetComponentGeometry(OutputComponent, SelectedComponentGeometry))
 		{
-			const FVector2f WidgetPosition = [Handle, &SelectedComponentGeometry]()
+			const FVector2f LocalPosition = [Handle, &SelectedComponentGeometry]()
 			{
 				// Get the initial offset based on the location around the selected object.
 				switch (Handle->GetTransformDirection())
 				{
 				case EDMXPixelMappingTransformDirection::CenterRight:
 					return FVector2f(SelectedComponentGeometry.GetLocalSize().X, SelectedComponentGeometry.GetLocalSize().Y * 0.5f);
-				case EDMXPixelMappingTransformDirection::BottomLeft:
-					return FVector2f(0, SelectedComponentGeometry.GetLocalSize().Y);
 				case EDMXPixelMappingTransformDirection::BottomCenter:
 					return FVector2f(SelectedComponentGeometry.GetLocalSize().X * 0.5f, SelectedComponentGeometry.GetLocalSize().Y);
 				case EDMXPixelMappingTransformDirection::BottomRight:
@@ -792,14 +785,30 @@ FVector2D SDMXPixelMappingDesignerView::GetExtensionPosition(TSharedPtr<SDMXPixe
 				return FVector2f::ZeroVector;
 			}();
 
+			const auto RotateVectorAroundPivot =
+				[](FVector2f Point, FVector2f Pivot, double AngleDegrees)
+				{
+					double Sin;
+					double Cos;
+					FMath::SinCos(&Sin, &Cos, FMath::DegreesToRadians(AngleDegrees));
+
+					const FVector2f RelativePositionOld = Point - Pivot;
+					const FVector2f RelativePositionNew = FVector2f(
+						Cos * RelativePositionOld.X - Sin * RelativePositionOld.Y,
+						Sin * RelativePositionOld.X + Cos * RelativePositionOld.Y);
+
+					return RelativePositionNew + Pivot;
+				};
+			const FVector2f LocalPositionRotated = RotateVectorAroundPivot(LocalPosition, SelectedComponentGeometry.GetLocalSize() / 2.f, OutputComponent->GetRotation());
+
 			const FVector2f SelectedWidgetScale = SelectedComponentGeometry.GetAccumulatedRenderTransform().GetMatrix().GetScale().GetVector();
 
 			const FVector2f ApplicationScaledOffset = UE::Slate::CastToVector2f(Handle->GetOffset() * GetDesignerGeometry().Scale);
 
 			const FVector2f LocalOffsetFull = ApplicationScaledOffset / SelectedWidgetScale;
-			const FVector2f PositionFullOffset = GetDesignerGeometry().AbsoluteToLocal(SelectedComponentGeometry.LocalToAbsolute(WidgetPosition + LocalOffsetFull));
+			const FVector2f PositionFullOffset = GetDesignerGeometry().AbsoluteToLocal(SelectedComponentGeometry.LocalToAbsolute(LocalPositionRotated + LocalOffsetFull));
 			const FVector2f LocalOffsetHalf = (ApplicationScaledOffset / 2.0f) / SelectedWidgetScale;
-			const FVector2f PositionHalfOffset = GetDesignerGeometry().AbsoluteToLocal(SelectedComponentGeometry.LocalToAbsolute(WidgetPosition + LocalOffsetHalf));
+			const FVector2f PositionHalfOffset = GetDesignerGeometry().AbsoluteToLocal(SelectedComponentGeometry.LocalToAbsolute(LocalPositionRotated + LocalOffsetHalf));
 
 			FVector2f PivotCorrection = PositionHalfOffset - (PositionFullOffset + FVector2f(5.0f, 5.0f));
 
@@ -1156,85 +1165,4 @@ void SDMXPixelMappingDesignerView::HandleDragEnterFromDetailsOrPalette(const TSh
 	}
 }
 
-SDMXPixelMappingDesignerView::FScopedRestoreSelection::FScopedRestoreSelection(TSharedRef<FDMXPixelMappingToolkit> ToolkitPtr, TSharedRef<SDMXPixelMappingDesignerView> DesignerView)
-	: WeakToolkit(ToolkitPtr)
-	, WeakDesignerView(DesignerView)
-{
-	if (TSharedPtr<FDMXPixelMappingToolkit> PinnedToolkit = WeakToolkit.Pin())
-	{
-		for (const FDMXPixelMappingComponentReference& ComponentReference : PinnedToolkit->GetSelectedComponents())
-		{
-			CachedSelectedComponents.Add(ComponentReference.GetComponent());
-		}
-	}
-}
-
-SDMXPixelMappingDesignerView::FScopedRestoreSelection::~FScopedRestoreSelection()
-{
-	if (TSharedPtr<FDMXPixelMappingToolkit> PinnedToolkit = WeakToolkit.Pin())
-	{
-		TArray<UDMXPixelMappingBaseComponent*> RemovedComponents;
-		TSet<FDMXPixelMappingComponentReference> ValidComponents;
-
-		if (UDMXPixelMapping* PixelMapping = PinnedToolkit->GetDMXPixelMapping())
-		{
-			TArray<UDMXPixelMappingBaseComponent*> ComponentsInPixelMapping;
-			PixelMapping->GetAllComponentsOfClass<UDMXPixelMappingBaseComponent>(ComponentsInPixelMapping);
-
-			for (UDMXPixelMappingBaseComponent* Component : ComponentsInPixelMapping)
-			{
-				const bool bComponentStillExists =
-					CachedSelectedComponents.ContainsByPredicate([Component](const TWeakObjectPtr<UDMXPixelMappingBaseComponent>& WeakCachedComponent)
-						{
-							return
-								WeakCachedComponent.IsValid() &&
-								Component == WeakCachedComponent.Get();
-						});
-
-				if (bComponentStillExists)
-				{
-					ValidComponents.Add(PinnedToolkit->GetReferenceFromComponent(Component));
-				}
-				else
-				{
-					RemovedComponents.Add(Component);
-				}
-			}
-
-			if (ValidComponents.Num() == 0)
-			{
-				// All were removed, select the the parent if possible or the renderer
-				UDMXPixelMappingBaseComponent** ComponentWithParentPtr = RemovedComponents.FindByPredicate([&ComponentsInPixelMapping](UDMXPixelMappingBaseComponent* Component)
-					{
-						return
-							Component &&
-							Component->GetParent() &&
-							ComponentsInPixelMapping.Contains(Component->GetParent());
-					});
-
-				if (ComponentWithParentPtr)
-				{
-					ValidComponents.Add(PinnedToolkit->GetReferenceFromComponent((*ComponentWithParentPtr)->GetParent()));
-				}
-				else
-				{
-					// Select the renderer
-					UDMXPixelMappingRendererComponent* RendererComponent = PinnedToolkit->GetActiveRendererComponent();
-					if (RendererComponent)
-					{
-						ValidComponents.Add(PinnedToolkit->GetReferenceFromComponent(RendererComponent));
-					}
-				}
-
-				PinnedToolkit->SelectComponents(ValidComponents);
-			}
-		}
-	}
-
-	if (TSharedPtr<SDMXPixelMappingDesignerView> PinnedDesignerView = WeakDesignerView.Pin())
-	{
-		PinnedDesignerView->CreateExtensionWidgetsForSelection();
-	}
-}
-
 #undef LOCTEXT_NAMESPACE
diff --git a/Engine/Plugins/VirtualProduction/DMX/DMXPixelMapping/Source/DMXPixelMappingEditor/Private/Views/SDMXPixelMappingDesignerView.h b/Engine/Plugins/VirtualProduction/DMX/DMXPixelMapping/Source/DMXPixelMappingEditor/Private/Views/SDMXPixelMappingDesignerView.h
--- a/Engine/Plugins/VirtualProduction/DMX/DMXPixelMapping/Source/DMXPixelMappingEditor/Private/Views/SDMXPixelMappingDesignerView.h
+++ b/Engine/Plugins/VirtualProduction/DMX/DMXPixelMapping/Source/DMXPixelMappingEditor/Private/Views/SDMXPixelMappingDesignerView.h
@@ -208,17 +208,4 @@ private:
 	TArray<TSharedRef<IDMXPixelMappingOutputComponentWidgetInterface>> OutputComponentWidgets;
 
 	TArray<TSharedPtr<SDMXPixelMappingTransformHandle>> TransformHandles;
-
-	/** Helper class to restore selection on scope */
-	class FScopedRestoreSelection
-	{
-	public:
-		FScopedRestoreSelection(TSharedRef<FDMXPixelMappingToolkit> ToolkitPtr, TSharedRef<SDMXPixelMappingDesignerView> DesignerView);
-		~FScopedRestoreSelection();
-
-	private:
-		TWeakPtr<FDMXPixelMappingToolkit> WeakToolkit;
-		TWeakPtr<SDMXPixelMappingDesignerView> WeakDesignerView;
-		TArray<TWeakObjectPtr<UDMXPixelMappingBaseComponent>> CachedSelectedComponents;
-	};
 };
diff --git a/Engine/Plugins/VirtualProduction/DMX/DMXPixelMapping/Source/DMXPixelMappingEditor/Private/Widgets/SDMXPixelMappingDesignerToolbar.cpp b/Engine/Plugins/VirtualProduction/DMX/DMXPixelMapping/Source/DMXPixelMappingEditor/Private/Widgets/SDMXPixelMappingDesignerToolbar.cpp
--- a/Engine/Plugins/VirtualProduction/DMX/DMXPixelMapping/Source/DMXPixelMappingEditor/Private/Widgets/SDMXPixelMappingDesignerToolbar.cpp
+++ b/Engine/Plugins/VirtualProduction/DMX/DMXPixelMapping/Source/DMXPixelMappingEditor/Private/Widgets/SDMXPixelMappingDesignerToolbar.cpp
@@ -7,6 +7,7 @@
 #include "Styling/AppStyle.h"
 #include "Styling/ToolBarStyle.h"
 #include "SViewportToolBarComboMenu.h"
+#include "ToolMenus.h"
 #include "Toolkits/DMXPixelMappingToolkit.h"
 #include "Widgets/Images/SImage.h"
 #include "Widgets/Input/SButton.h"
@@ -17,66 +18,174 @@
 
 #define LOCTEXT_NAMESPACE "SDMXPixelMappingDesignerToolbar"
 
+
 namespace UE::DMX
 {
-	void SDMXPixelMappingDesignerToolbar::Construct(const FArguments& InArgs, const TSharedRef<FDMXPixelMappingToolkit>& InToolkit)
+	namespace DMXPixelMappingDesignerToolbar::Private
 	{
-		WeakToolkit = InToolkit;
+		constexpr TCHAR ToolbarName[] = TEXT("PixelMapping.DesignerToolbar");
+	}
 
-		constexpr TCHAR ToolBarStyleName[] = TEXT("EditorViewportToolBar");
-		const FToolBarStyle& ToolBarStyle = FAppStyle::Get().GetWidgetStyle<FToolBarStyle>(ToolBarStyleName);
+	SDMXPixelMappingDesignerToolbar::~SDMXPixelMappingDesignerToolbar()
+	{
+		using namespace DMXPixelMappingDesignerToolbar::Private;
+		if (UToolMenus::Get()->IsMenuRegistered(ToolbarName))
+		{
+			UToolMenus::Get()->RemoveMenu(ToolbarName);
+		}
+	}
 
-		const TSharedRef<FUICommandList> CommandList = InToolkit->GetToolkitCommands();
+	void SDMXPixelMappingDesignerToolbar::Construct(const FArguments& InArgs, const TSharedRef<FDMXPixelMappingToolkit>& InToolkit)
+	{
+		WeakToolkit = InToolkit;
 
-		const FUICommandInfo* ToggleSnapGridCommand = FDMXPixelMappingEditorCommands::Get().ToggleGridSnapping.Get();
-		check(ToggleSnapGridCommand);
+		RegisterToolbarMenu(InArgs);
+		
+		const FToolMenuContext Context(InToolkit->GetToolkitCommands());
 
+		using namespace DMXPixelMappingDesignerToolbar::Private;
 		ChildSlot
 		[
-			SNew(SBorder)
-			.BorderImage(FAppStyle::GetBrush("NoBorder"))
-			.ForegroundColor(FAppStyle::GetSlateColor("DefaultForeground"))
-			[
-				SNew(SHorizontalBox)
+			UToolMenus::Get()->GenerateWidget(ToolbarName, Context)
+		];
+	}
+
+	void SDMXPixelMappingDesignerToolbar::RegisterToolbarMenu(const FArguments& InArgs)
+	{
+		using namespace UE::DMX;
 
-				// Grid snapping
+		UToolMenus* ToolMenus = UToolMenus::Get();
+
+		using namespace DMXPixelMappingDesignerToolbar::Private;
+		if (ToolMenus->IsMenuRegistered(ToolbarName))
+		{
+			return;
+		}
+
+		UToolMenu* Toolbar = UToolMenus::Get()->RegisterMenu(ToolbarName, NAME_None, EMultiBoxType::SlimHorizontalToolBar);
+		Toolbar->StyleName = "EditorViewportToolbar";
+		
+
+		// Transform Handle Modes
+		{
+			FToolMenuSection& Section = Toolbar->AddSection("TransformHandleMode");
+
+			// Resize mode
+			const FCheckBoxStyle& CheckBoxStartStyle = FAppStyle::Get().GetWidgetStyle<FCheckBoxStyle>("EditorViewportToolBar.ToggleButton.Start");
+
+			TSharedRef<SCheckBox> EnableResizeModeToggleButton = 
+				SNew(SCheckBox)
+				.Style(&CheckBoxStartStyle)
+				.ToolTipText(LOCTEXT("TransformHandleResizeMode", "Resize Components"))
+				.OnCheckStateChanged(this, &SDMXPixelMappingDesignerToolbar::OnTransformHandleModeSelected, EDMXPixelMappingTransformHandleMode::Resize)
+				.IsChecked(this, &SDMXPixelMappingDesignerToolbar::GetCheckboxStateForTransormHandleMode, EDMXPixelMappingTransformHandleMode::Resize)
+				[
+					SNew(SImage)
+					.Image(FAppStyle::GetBrush("EditorViewport.ScaleMode"))
+					.ColorAndOpacity(FSlateColor::UseForeground())
+				];
+
+			// Rotate mode
+			const FCheckBoxStyle& CheckBoxEndStyle = FAppStyle::Get().GetWidgetStyle<FCheckBoxStyle>("EditorViewportToolBar.ToggleButton.End");
+
+			TSharedRef<SCheckBox> EnableRotateModeToggleButton = 
+				SNew(SCheckBox)
+				.Style(&CheckBoxEndStyle)
+				.ToolTipText(LOCTEXT("TransformHandleRotateMode", "Rotate Components"))
+				.OnCheckStateChanged(this, &SDMXPixelMappingDesignerToolbar::OnTransformHandleModeSelected, EDMXPixelMappingTransformHandleMode::Rotate)
+				.IsChecked(this, &SDMXPixelMappingDesignerToolbar::GetCheckboxStateForTransormHandleMode, EDMXPixelMappingTransformHandleMode::Rotate)
+				[
+					SNew(SImage)
+					.Image(FAppStyle::GetBrush("EditorViewport.RotateMode"))
+					.ColorAndOpacity(FSlateColor::UseForeground())
+				];
+
+			// As a single widget
+			const TSharedRef<SWidget> TransformHandleModeWidget =
+				SNew(SHorizontalBox)
 				+ SHorizontalBox::Slot()
-				.AutoWidth()
-				.VAlign(VAlign_Center)
-				.Padding(4.f, 2.0f)
+				.Padding(0.f)
 				[
-					SNew(SViewportToolBarComboMenu)
-					.Cursor(EMouseCursor::Default)
-					.IsChecked(this, &SDMXPixelMappingDesignerToolbar::GetSnapGridEnabledCheckState)
-					.OnCheckStateChanged(this, &SDMXPixelMappingDesignerToolbar::OnSnapGridCheckStateChanged)
-					.Label(this, &SDMXPixelMappingDesignerToolbar::GetSnapGridLabel)
-					.OnGetMenuContent(this, &SDMXPixelMappingDesignerToolbar::GenerateSnapGridMenu)
-					.ToggleButtonToolTip(ToggleSnapGridCommand->GetDescription())
-					.MenuButtonToolTip(LOCTEXT("SnapGridMenuTooltip", "Grid Snapping Settings"))
-					.Icon(FSlateIcon(FAppStyle::GetAppStyleSetName(), "EditorViewport.LocationGridSnap"))
-					.ParentToolBar(SharedThis(this))
+					EnableResizeModeToggleButton
 				]
-
-				// Zoom to fit
 				+ SHorizontalBox::Slot()
-				.AutoWidth()
+				.Padding(0.f)
+				[
+					EnableRotateModeToggleButton
+				];
+
+			Section.AddEntry
+			(
+				FToolMenuEntry::InitWidget
+				(
+					"TransformHandleModes",
+					TransformHandleModeWidget,
+					FText::GetEmpty()
+				)
+			);
+		}
+
+
+		// Grid Snapping
+		{
+			FToolMenuSection& Section = Toolbar->AddSection("GridSnapping");
+
+			const FUICommandInfo& ToggleSnapGridCommand = *FDMXPixelMappingEditorCommands::Get().ToggleGridSnapping;
+
+			const TSharedRef<SViewportToolBarComboMenu> GridSnappingComboMenu =
+				SNew(SViewportToolBarComboMenu)
+				.Cursor(EMouseCursor::Default)
+				.IsChecked(this, &SDMXPixelMappingDesignerToolbar::GetSnapGridEnabledCheckState)
+				.OnCheckStateChanged(this, &SDMXPixelMappingDesignerToolbar::OnSnapGridCheckStateChanged)
+				.Label(this, &SDMXPixelMappingDesignerToolbar::GetSnapGridLabel)
+				.OnGetMenuContent(this, &SDMXPixelMappingDesignerToolbar::GenerateSnapGridMenu)
+				.ToggleButtonToolTip(ToggleSnapGridCommand.GetDescription())
+				.MenuButtonToolTip(LOCTEXT("SnapGridMenuTooltip", "Grid Snapping Settings"))
+				.Icon(FSlateIcon(FAppStyle::GetAppStyleSetName(), "EditorViewport.LocationGridSnap"))
+				.ParentToolBar(SharedThis(this));
+
+			Section.AddEntry
+			(
+				FToolMenuEntry::InitWidget
+				(
+					"GridSnappingComboButton",
+					GridSnappingComboMenu,
+					FText::GetEmpty()
+				)
+			);
+		}
+
+
+		// Zoom to Fit
+		{
+			FToolMenuSection& Section = Toolbar->AddSection("ZoomToFit");
+
+			constexpr TCHAR ToolBarStyleName[] = TEXT("EditorViewportToolBar");
+			const FToolBarStyle& ToolBarStyle = FAppStyle::Get().GetWidgetStyle<FToolBarStyle>(ToolBarStyleName);
+
+			const TSharedRef<SButton> ZoomToFitButton = 
+				SNew(SButton)
+				.ButtonStyle(&ToolBarStyle.ButtonStyle)
+				.ToolTipText(LOCTEXT("ZoomToFit_ToolTip", "Zoom To Fit"))
+				.OnClicked(InArgs._OnZoomToFitClicked)
+				.ContentPadding(ToolBarStyle.ButtonPadding)
 				.VAlign(VAlign_Center)
-				.Padding(4.f, 2.0f)
 				[
-					SNew(SButton)
-					.ButtonStyle(&ToolBarStyle.ButtonStyle)
-					.ToolTipText(LOCTEXT("ZoomToFit_ToolTip", "Zoom To Fit"))
-					.OnClicked(InArgs._OnZoomToFitClicked)
-					.ContentPadding(ToolBarStyle.ButtonPadding)
-					.VAlign(VAlign_Center)
-					[
-						SNew(SImage)
-						.Image(FAppStyle::GetBrush("UMGEditor.ZoomToFit"))
-						.ColorAndOpacity(FSlateColor::UseForeground())
-					]
-				]
-			]
-		];
+					SNew(SImage)
+					.Image(FAppStyle::GetBrush("UMGEditor.ZoomToFit"))
+					.ColorAndOpacity(FSlateColor::UseForeground())
+				];
+
+			Section.AddEntry
+			(
+				FToolMenuEntry::InitWidget
+				(
+					"ZoomToFitButton",
+					ZoomToFitButton,
+					FText::GetEmpty()
+				)
+			);
+		}
 	}
 
 	TSharedRef<SWidget> SDMXPixelMappingDesignerToolbar::GenerateSnapGridMenu()
@@ -131,7 +240,28 @@ namespace UE::DMX
 			Toolkit->ToggleGridSnapping();
 		}
 	}
-}
 
+	void SDMXPixelMappingDesignerToolbar::OnTransformHandleModeSelected(ECheckBoxState DummyCheckBoxState, UE::DMX::EDMXPixelMappingTransformHandleMode NewTransformHandleMode)
+	{
+		const TSharedPtr<FDMXPixelMappingToolkit> Toolkit = WeakToolkit.Pin();
+		if (!Toolkit.IsValid())
+		{
+			return;
+		}
+
+		Toolkit->SetTransformHandleMode(NewTransformHandleMode);
+	}
+
+	ECheckBoxState SDMXPixelMappingDesignerToolbar::GetCheckboxStateForTransormHandleMode(UE::DMX::EDMXPixelMappingTransformHandleMode TransformHandleMode) const
+	{
+		const TSharedPtr<FDMXPixelMappingToolkit> Toolkit = WeakToolkit.Pin();
+		if (!Toolkit.IsValid())
+		{
+			return ECheckBoxState::Undetermined;
+		}
+
+		return Toolkit->GetTransformHandleMode() == TransformHandleMode ? ECheckBoxState::Checked : ECheckBoxState::Unchecked;
+	}
+}
 
 #undef LOCTEXT_NAMESPACE
diff --git a/Engine/Plugins/VirtualProduction/DMX/DMXPixelMapping/Source/DMXPixelMappingEditor/Private/Widgets/SDMXPixelMappingDesignerToolbar.h b/Engine/Plugins/VirtualProduction/DMX/DMXPixelMapping/Source/DMXPixelMappingEditor/Private/Widgets/SDMXPixelMappingDesignerToolbar.h
--- a/Engine/Plugins/VirtualProduction/DMX/DMXPixelMapping/Source/DMXPixelMappingEditor/Private/Widgets/SDMXPixelMappingDesignerToolbar.h
+++ b/Engine/Plugins/VirtualProduction/DMX/DMXPixelMapping/Source/DMXPixelMappingEditor/Private/Widgets/SDMXPixelMappingDesignerToolbar.h
@@ -6,7 +6,9 @@
 #include "Styling/SlateTypes.h"
 #include "SViewportToolBar.h"
 
+enum class ECheckBoxState : uint8;
 class FDMXPixelMappingToolkit;
+namespace UE::DMX { enum class EDMXPixelMappingTransformHandleMode : uint8; }
 
 
 namespace UE::DMX
@@ -18,14 +20,20 @@ namespace UE::DMX
 		SLATE_BEGIN_ARGS(SDMXPixelMappingDesignerToolbar)
 			{}
 
+			/** Event raised when zoom to fit was clicked */
 			SLATE_EVENT(FOnClicked, OnZoomToFitClicked)
 
 		SLATE_END_ARGS()
 
+		~SDMXPixelMappingDesignerToolbar();
+
 		/** Constructs this widget */
 		void Construct(const FArguments& InArgs, const TSharedRef<FDMXPixelMappingToolkit>& InToolkit);
 
 	private:
+		/** Builds the toolbar for this widget */
+		void RegisterToolbarMenu(const FArguments& InArgs);
+
 		/** Generates the grid snapping menu */
 		TSharedRef<SWidget> GenerateSnapGridMenu();
 
@@ -38,6 +46,12 @@ namespace UE::DMX
 		/** Called when the grid snapping checkbox state changed */
 		void OnSnapGridCheckStateChanged(ECheckBoxState NewCheckBoxState);
 
+		/** Called when a transform handle mode was selected. The checkbox state can be ignored. Instead clicking always enables. */
+		void OnTransformHandleModeSelected(ECheckBoxState DummyCheckBoxState, UE::DMX::EDMXPixelMappingTransformHandleMode NewTransformHandleMode);
+
+		/** Returns the current checkbox state of the transform mode */
+		ECheckBoxState GetCheckboxStateForTransormHandleMode(UE::DMX::EDMXPixelMappingTransformHandleMode TransformHandleMode) const;
+
 		/** The toolkit that owns this widget */
 		TWeakPtr<FDMXPixelMappingToolkit> WeakToolkit;
 	};
diff --git a/Engine/Plugins/VirtualProduction/DMX/DMXPixelMapping/Source/DMXPixelMappingEditor/Private/Widgets/SDMXPixelMappingFixturePatchList.cpp b/Engine/Plugins/VirtualProduction/DMX/DMXPixelMapping/Source/DMXPixelMappingEditor/Private/Widgets/SDMXPixelMappingFixturePatchList.cpp
--- a/Engine/Plugins/VirtualProduction/DMX/DMXPixelMapping/Source/DMXPixelMappingEditor/Private/Widgets/SDMXPixelMappingFixturePatchList.cpp
+++ b/Engine/Plugins/VirtualProduction/DMX/DMXPixelMapping/Source/DMXPixelMappingEditor/Private/Widgets/SDMXPixelMappingFixturePatchList.cpp
@@ -45,14 +45,10 @@ FReply SDMXPixelMappingFixturePatchList::OnRowDragDetected(const FGeometry& MyGe
 	TArray<TSharedPtr<FDMXPixelMappingComponentTemplate>> Templates;
 	for (UDMXEntityFixturePatch* FixturePatch : GetSelectedFixturePatches())
 	{
-		UDMXEntityFixtureType* FixtureType = FixturePatch ? FixturePatch->GetFixtureType() : nullptr;
-		if (!FixturePatch || !FixtureType)
-		{
-			continue;
-		}
+		const FDMXFixtureMode* FixtureModePtr = FixturePatch->GetActiveMode();
 
 		const FDMXEntityFixturePatchRef FixturePatchReference(FixturePatch);
-		if (FixtureType->bFixtureMatrixEnabled)
+		if (FixtureModePtr && FixtureModePtr->bFixtureMatrixEnabled)
 		{
 			TSharedRef<FDMXPixelMappingComponentTemplate> FixturePatchMatrixTemplate = MakeShared<FDMXPixelMappingComponentTemplate>(UDMXPixelMappingMatrixComponent::StaticClass(), FixturePatchReference);
 			Templates.Add(FixturePatchMatrixTemplate);
diff --git a/Engine/Plugins/VirtualProduction/DMX/DMXPixelMapping/Source/DMXPixelMappingEditor/Private/Widgets/SDMXPixelMappingOutputComponent.cpp b/Engine/Plugins/VirtualProduction/DMX/DMXPixelMapping/Source/DMXPixelMappingEditor/Private/Widgets/SDMXPixelMappingOutputComponent.cpp
--- a/Engine/Plugins/VirtualProduction/DMX/DMXPixelMapping/Source/DMXPixelMappingEditor/Private/Widgets/SDMXPixelMappingOutputComponent.cpp
+++ b/Engine/Plugins/VirtualProduction/DMX/DMXPixelMapping/Source/DMXPixelMappingEditor/Private/Widgets/SDMXPixelMappingOutputComponent.cpp
@@ -30,7 +30,7 @@ void IDMXPixelMappingOutputComponentWidgetInterface::AddToCanvas(const TSharedRe
 	ParentCanvas = InCanvas;
 
 	ParentCanvas->AddSlot()
-		.ZOrder(100)
+		.ZOrder(0)
 		.AutoSize(true)
 		.Alignment(FVector2D::ZeroVector)
 		.Offset_Lambda([this]()
@@ -123,6 +123,8 @@ TSharedRef<SWidget> SDMXPixelMappingOutputComponent::CreateContent()
 {
 	const TSharedRef<SBorder> Content =
 		SNew(SBorder)
+		.RenderTransform(this, &SDMXPixelMappingOutputComponent::ComputeRenderTransform)
+		.RenderTransformPivot(FVector2D(0.5, 0.5))
 		.BorderImage(FDMXPixelMappingEditorStyle::Get().GetBrush("DMXPixelMappingEditor.ComponentBorder"))
 		.BorderBackgroundColor_Lambda([this]()
 			{
@@ -319,6 +321,10 @@ void SDMXPixelMappingOutputComponent::CreatePatchInfoChildSlot()
 	);
 }
 
+TOptional<FSlateRenderTransform> SDMXPixelMappingOutputComponent::ComputeRenderTransform() const
+{
+	return FSlateRenderTransform(Model->GetQuaternion(), FVector2D(0.0, 0.0));
+}
 
 class SDMXPixelMappingScreenComponentCell
 	: public SCompoundWidget
diff --git a/Engine/Plugins/VirtualProduction/DMX/DMXPixelMapping/Source/DMXPixelMappingEditor/Private/Widgets/SDMXPixelMappingOutputComponent.h b/Engine/Plugins/VirtualProduction/DMX/DMXPixelMapping/Source/DMXPixelMappingEditor/Private/Widgets/SDMXPixelMappingOutputComponent.h
--- a/Engine/Plugins/VirtualProduction/DMX/DMXPixelMapping/Source/DMXPixelMappingEditor/Private/Widgets/SDMXPixelMappingOutputComponent.h
+++ b/Engine/Plugins/VirtualProduction/DMX/DMXPixelMapping/Source/DMXPixelMappingEditor/Private/Widgets/SDMXPixelMappingOutputComponent.h
@@ -92,6 +92,9 @@ private:
 	/** Creates the child slot that displays info about the patch such as Addresses or the Fixture ID */
 	void CreatePatchInfoChildSlot();
 
+	/** Returns the render transform for this widget */
+	TOptional<FSlateRenderTransform> ComputeRenderTransform() const;
+
 	/** The box that is shown */
 	TSharedPtr<SBox> ComponentBox;
 
diff --git a/Engine/Plugins/VirtualProduction/DMX/DMXPixelMapping/Source/DMXPixelMappingEditor/Private/Widgets/SDMXPixelMappingTransformHandle.cpp b/Engine/Plugins/VirtualProduction/DMX/DMXPixelMapping/Source/DMXPixelMappingEditor/Private/Widgets/SDMXPixelMappingTransformHandle.cpp
--- a/Engine/Plugins/VirtualProduction/DMX/DMXPixelMapping/Source/DMXPixelMappingEditor/Private/Widgets/SDMXPixelMappingTransformHandle.cpp
+++ b/Engine/Plugins/VirtualProduction/DMX/DMXPixelMapping/Source/DMXPixelMappingEditor/Private/Widgets/SDMXPixelMappingTransformHandle.cpp
@@ -22,7 +22,6 @@ void SDMXPixelMappingTransformHandle::Construct(const FArguments& InArgs, TShare
 	Offset = InOffset;
 
 	Action = EDMXPixelMappingTransformAction::None;
-	ScopedTransaction = nullptr;
 
 	DragDirection = ComputeDragDirection(InTransformDirection);
 	DragOrigin = ComputeOrigin(InTransformDirection);
@@ -42,45 +41,46 @@ EVisibility SDMXPixelMappingTransformHandle::GetHandleVisibility() const
 
 FReply SDMXPixelMappingTransformHandle::OnMouseButtonDown(const FGeometry& MyGeometry, const FPointerEvent& MouseEvent)
 {
-	if (MouseEvent.GetEffectingButton() == EKeys::LeftMouseButton)
+	const TSharedPtr<FDMXPixelMappingToolkit> Toolkit = DesignerViewWeakPtr.IsValid() ? DesignerViewWeakPtr.Pin()->GetToolkit() : nullptr;
+	if (!Toolkit.IsValid() || MouseEvent.GetEffectingButton() != EKeys::LeftMouseButton)
 	{
-		const FVector2D LocalSize = MyGeometry.GetLocalSize();
-		const FVector2D LocalCursorPos = MyGeometry.AbsoluteToLocal(MouseEvent.GetScreenSpacePosition());
+		return FReply::Unhandled();
+	}
 
-		// Only handle dragging the edges of the widget
-		constexpr double MouseThreshold = 10.f;
-		if (LocalCursorPos.X < LocalSize.X - MouseThreshold ||
-			LocalCursorPos.Y < LocalSize.Y - MouseThreshold)
-		{
-			return FReply::Unhandled();;
-		}
+	const FVector2D LocalSize = MyGeometry.GetLocalSize();
+	const FVector2D LocalCursorPos = MyGeometry.AbsoluteToLocal(MouseEvent.GetScreenSpacePosition());
 
-		Action = ComputeActionAtLocation(MyGeometry, MouseEvent);
+	// Only handle dragging the edges of the widget
+	constexpr double MouseThreshold = 10.0;
+	if (LocalCursorPos.X < LocalSize.X - MouseThreshold ||
+		LocalCursorPos.Y < LocalSize.Y - MouseThreshold)
+	{
+		return FReply::Unhandled();;
+	}
 
-		if (TSharedPtr<FDMXPixelMappingToolkit> Toolkit = DesignerViewWeakPtr.Pin()->GetToolkit())
-		{
-			const TSet<FDMXPixelMappingComponentReference> SelectedComponentReferences = Toolkit->GetSelectedComponents();
-			for (const FDMXPixelMappingComponentReference& ComponentReference : SelectedComponentReferences)
-			{
-				UDMXPixelMappingBaseComponent* Component = ComponentReference.GetComponent();
+	Action = ComputeActionAtLocation(MyGeometry, MouseEvent);
 
-				if (UDMXPixelMappingOutputComponent* OutputComponent = Cast<UDMXPixelMappingOutputComponent>(Component))
-				{
-					FMargin Offsets;
-					FVector2D Size = OutputComponent->GetSize();
-					Offsets.Right = Size.X;
-					Offsets.Bottom = Size.Y;
-					StartingOffsets = Offsets;
-				}
+	const TSet<FDMXPixelMappingComponentReference> SelectedComponentReferences = Toolkit->GetSelectedComponents();
+	for (const FDMXPixelMappingComponentReference& ComponentReference : SelectedComponentReferences)
+	{
+		UDMXPixelMappingBaseComponent* Component = ComponentReference.GetComponent();
 
-				MouseDownPosition = MouseEvent.GetScreenSpacePosition();
+		if (UDMXPixelMappingOutputComponent* OutputComponent = Cast<UDMXPixelMappingOutputComponent>(Component))
+		{
+			InitialSize = OutputComponent->GetSize();
+			InitialRotation = OutputComponent->GetRotation();
+		}
 
-				ScopedTransaction = MakeShareable<FScopedTransaction>(new FScopedTransaction(LOCTEXT("ResizePixelMappingComponent", "PixelMapping: Resize Component")));
-				Component->Modify();
+		MouseDownPosition = MouseEvent.GetScreenSpacePosition();
 
-				return FReply::Handled().CaptureMouse(SharedThis(this));
-			}
-		}
+		const FText TransactionText = Toolkit->GetTransformHandleMode() == UE::DMX::EDMXPixelMappingTransformHandleMode::Resize ?
+			LOCTEXT("ResizePixelMappingComponent", "PixelMapping: Resize Component") :
+			LOCTEXT("RotatePixelMappingComponent", "PixelMapping: Rotate Component");
+	
+		ScopedTransaction = MakeShareable<FScopedTransaction>(new FScopedTransaction(TransactionText));
+		Component->Modify();
+
+		return FReply::Handled().CaptureMouse(SharedThis(this));
 	}
 
 	return FReply::Unhandled();
@@ -88,29 +88,40 @@ FReply SDMXPixelMappingTransformHandle::OnMouseButtonDown(const FGeometry& MyGeo
 
 FReply SDMXPixelMappingTransformHandle::OnMouseButtonUp(const FGeometry& MyGeometry, const FPointerEvent& MouseEvent)
 {
-	if ( HasMouseCapture() && MouseEvent.GetEffectingButton() == EKeys::LeftMouseButton )
+	RequestApplyTransformHandle.Invalidate();
+
+	const TSharedPtr<FDMXPixelMappingToolkit> Toolkit = DesignerViewWeakPtr.IsValid() ? DesignerViewWeakPtr.Pin()->GetToolkit() : nullptr;
+	if (!Toolkit.IsValid() || !HasMouseCapture() || MouseEvent.GetEffectingButton() != EKeys::LeftMouseButton)
 	{
-		Action = EDMXPixelMappingTransformAction::None;
+		return FReply::Unhandled();
+	}
 		
-		if (TSharedPtr<FDMXPixelMappingToolkit> Toolkit = DesignerViewWeakPtr.Pin()->GetToolkit())
-		{		
-			const TSet<FDMXPixelMappingComponentReference> SelectedComponentReferences = Toolkit->GetSelectedComponents();
-			for (const FDMXPixelMappingComponentReference& ComponentReference : SelectedComponentReferences)
-			{
-				if (UDMXPixelMappingOutputComponent* ResizedComponent = Cast<UDMXPixelMappingOutputComponent>(ComponentReference.GetComponent()))
-				{
-					// Set the final size transacted
-					const FVector2D Delta = MouseEvent.GetScreenSpacePosition() - MouseDownPosition;
-					const FVector2D TranslateAmount = Delta * (1.0f / (DesignerViewWeakPtr.Pin()->GetZoomAmount() * MyGeometry.Scale));
-
-					ResizedComponent->Modify();
-					Resize(ResizedComponent, DragDirection, TranslateAmount);
+	Action = EDMXPixelMappingTransformAction::None;
 
-					ScopedTransaction.Reset();
+	const TSet<FDMXPixelMappingComponentReference> SelectedComponentReferences = Toolkit->GetSelectedComponents();
+	for (const FDMXPixelMappingComponentReference& ComponentReference : SelectedComponentReferences)
+	{
+		if (UDMXPixelMappingOutputComponent* EditedComponent = Cast<UDMXPixelMappingOutputComponent>(ComponentReference.GetComponent()))
+		{
+			const FVector2D& CursorPosition = MouseEvent.GetScreenSpacePosition();
 
-					return FReply::Handled().ReleaseMouseCapture();
-				}
+			EditedComponent->Modify();
+			if (Toolkit->GetTransformHandleMode() == UE::DMX::EDMXPixelMappingTransformHandleMode::Resize)
+			{
+				Resize(EditedComponent, CursorPosition, MyGeometry.Scale);
+			}
+			else if (Toolkit->GetTransformHandleMode() == UE::DMX::EDMXPixelMappingTransformHandleMode::Rotate)
+			{
+				Rotate(EditedComponent, CursorPosition, MyGeometry.Scale);
 			}
+			else
+			{
+				checkf(0, TEXT("Unhandled enum value"));
+			}
+
+			ScopedTransaction.Reset();
+
+			return FReply::Handled().ReleaseMouseCapture();
 		}
 	}
 
@@ -119,19 +130,45 @@ FReply SDMXPixelMappingTransformHandle::OnMouseButtonUp(const FGeometry& MyGeome
 
 FReply SDMXPixelMappingTransformHandle::OnMouseMove(const FGeometry& MyGeometry, const FPointerEvent& MouseEvent)
 {
-	if (Action != EDMXPixelMappingTransformAction::None && DesignerViewWeakPtr.IsValid())
+	TSharedPtr<FDMXPixelMappingToolkit> Toolkit = DesignerViewWeakPtr.IsValid() ? DesignerViewWeakPtr.Pin()->GetToolkit() : nullptr;
+	if (!Toolkit.IsValid() || Action == EDMXPixelMappingTransformAction::None)
 	{
-		if (TSharedPtr<FDMXPixelMappingToolkit> Toolkit = DesignerViewWeakPtr.Pin()->GetToolkit())
+		return FReply::Unhandled();
+	}
+
+	if (RequestApplyTransformHandle.IsValid())
+	{
+		return FReply::Unhandled();
+	}
+
+	const TSet<FDMXPixelMappingComponentReference> SelectedComponentReferences = Toolkit->GetSelectedComponents();
+	for (const FDMXPixelMappingComponentReference& ComponentReference : SelectedComponentReferences)
+	{
+		UDMXPixelMappingOutputComponent* EditedComponent = Cast<UDMXPixelMappingOutputComponent>(ComponentReference.GetComponent());
+		if (!EditedComponent)
 		{
-			const TSet<FDMXPixelMappingComponentReference> SelectedComponentReferences = Toolkit->GetSelectedComponents();
-			for (const FDMXPixelMappingComponentReference& ComponentReference : SelectedComponentReferences)
-			{
-				if (UDMXPixelMappingOutputComponent* ResizedComponent = Cast<UDMXPixelMappingOutputComponent>(ComponentReference.GetComponent()))
+			continue;
+		}
+
+		const FVector2D& CursorPosition = MouseEvent.GetScreenSpacePosition();
+		RequestApplyTransformHandle = GEditor->GetTimerManager()->SetTimerForNextTick(FTimerDelegate::CreateLambda(
+			[this, EditedComponent, CursorPosition, MyGeometry, Toolkit]()
+			{					
+				if (Toolkit->GetTransformHandleMode() == UE::DMX::EDMXPixelMappingTransformHandleMode::Resize)
 				{
-					RequestResize(ResizedComponent, DragDirection);
+					Resize(EditedComponent, CursorPosition, MyGeometry.Scale);
 				}
-			}
-		}
+				else if (Toolkit->GetTransformHandleMode() == UE::DMX::EDMXPixelMappingTransformHandleMode::Rotate)
+				{
+					Rotate(EditedComponent, CursorPosition, MyGeometry.Scale);
+				}
+				else
+				{
+					checkf(0, TEXT("Unhandled enum value"));
+				}
+
+				RequestApplyTransformHandle.Invalidate();
+			}));
 	}
 
 	return FReply::Unhandled();
@@ -143,8 +180,6 @@ FCursorReply SDMXPixelMappingTransformHandle::OnCursorQuery(const FGeometry& MyG
 	{
 	case EDMXPixelMappingTransformDirection::BottomRight:
 		return FCursorReply::Cursor(EMouseCursor::ResizeSouthEast);
-	case EDMXPixelMappingTransformDirection::BottomLeft:
-		return FCursorReply::Cursor(EMouseCursor::ResizeSouthWest);
 	case EDMXPixelMappingTransformDirection::BottomCenter:
 		return FCursorReply::Cursor(EMouseCursor::ResizeUpDown);
 	case EDMXPixelMappingTransformDirection::CenterRight:
@@ -154,96 +189,58 @@ FCursorReply SDMXPixelMappingTransformHandle::OnCursorQuery(const FGeometry& MyG
 	return FCursorReply::Unhandled();
 }
 
-void SDMXPixelMappingTransformHandle::RequestResize(UDMXPixelMappingBaseComponent* BaseComponent, const FVector2D& Direction)
-{
-	if (!RequestResizeHandle.IsValid() && DesignerViewWeakPtr.IsValid())
+void SDMXPixelMappingTransformHandle::Rotate(UDMXPixelMappingOutputComponent* OutputComponent, const FVector2D& CursorPosition, double WidgetScale)
+{	
+	if (!OutputComponent)
 	{
-		const FVector2D& CursorPosition = FSlateApplication::Get().GetCursorPos();
-		const FVector2D Delta = CursorPosition - MouseDownPosition;
-		const FVector2D TranslateAmount = Delta * (1.0f / (DesignerViewWeakPtr.Pin()->GetZoomAmount() * GetCachedGeometry().Scale));
-
-		RequestResizeHandle = GEditor->GetTimerManager()->SetTimerForNextTick(FTimerDelegate::CreateLambda([this, BaseComponent, Direction, TranslateAmount]()
-			{
-				Resize(BaseComponent, Direction, TranslateAmount);
-				RequestResizeHandle.Invalidate();
-			}));
+		return;
 	}
 
+	// Position of the initially clicked location to the current mouse position in slate space. Remove zoom and scaling to get to graph space.
+	const FVector2D MouseDownToCurorPosition = CursorPosition - MouseDownPosition;
+	const FVector2D MouseDownToCursorPositionGraphSpace = MouseDownToCurorPosition / DesignerViewWeakPtr.Pin()->GetZoomAmount() / WidgetScale;
+	
+	const FVector2D InitialCenterToBottomRight = OutputComponent->GetSize().GetRotated(InitialRotation) / 2.0;
+	const FVector2D NewCenterToBottomRight = InitialCenterToBottomRight + MouseDownToCursorPositionGraphSpace;
+
+	const double DotProdcutNormalized = FVector2D::DotProduct(InitialCenterToBottomRight.GetSafeNormal(), NewCenterToBottomRight.GetSafeNormal());
+	const double Sign = FVector2D::CrossProduct(InitialCenterToBottomRight, NewCenterToBottomRight) > 0.0 ? 1.0 : -1.0;
+	const double NewAngle = Sign * FMath::UnwindRadians(FMath::Acos(DotProdcutNormalized));
+
+	OutputComponent->Modify();
+	OutputComponent->SetRotation(InitialRotation + FMath::RadiansToDegrees(NewAngle));
 }
 
-void SDMXPixelMappingTransformHandle::Resize(UDMXPixelMappingBaseComponent* BaseComponent, const FVector2D& Direction, const FVector2D& Amount)
+void SDMXPixelMappingTransformHandle::Resize(UDMXPixelMappingOutputComponent* OutputComponent, const FVector2D& CursorPosition, double WidgetScale)
 {
-	if (UDMXPixelMappingOutputComponent* OutputComponent = Cast<UDMXPixelMappingOutputComponent>(BaseComponent))
+	if (!OutputComponent)
 	{
-		FMargin Offsets = StartingOffsets;
-
-		const FVector2D Movement = Amount * Direction;
-		if (Direction.X < 0)
-		{
-			Offsets.Left -= Movement.X;
-			Offsets.Right += Movement.X;
-		}
-
-		if (Direction.Y < 0)
-		{
-			Offsets.Top -= Movement.Y;
-			Offsets.Bottom += Movement.Y;
-		}
-
-		if (Direction.X > 0)
-		{
-			Offsets.Left += Movement.X;
-			Offsets.Right += Amount.X * Direction.X;
-		}
-
-		if (Direction.Y > 0)
-		{
-			Offsets.Top += Movement.Y;
-			Offsets.Bottom += Amount.Y * Direction.Y;
-		}
+		return;
+	}
 
-		const FVector2D OldSize = OutputComponent->GetSize();
-		const FVector2D RequestedSize = FVector2D(Offsets.Right, Offsets.Bottom);
-		const FVector2D NewSize = GetSnapSize(OutputComponent, RequestedSize, Direction);
-		if (OldSize == NewSize)
-		{
-			// No unchanged values
-			return;
-		}
-		OutputComponent->Modify();
-		OutputComponent->SetSize(NewSize);
+	// Position of the initially clicked location to the current mouse position in slate space. Remove zoom and scaling to get to graph space.
+	const FVector2D MouseDownToCurorPosition = CursorPosition - MouseDownPosition;
+	const FVector2D MouseDownToCursorPositionGraphSpace = MouseDownToCurorPosition / DesignerViewWeakPtr.Pin()->GetZoomAmount() / WidgetScale;
 
-		// Scale children only if desired, no division by zero
-		const FDMXPixelMappingDesignerSettings& DesignerSettings = GetDefault<UDMXPixelMappingEditorSettings>()->DesignerSettings;
-		if (!DesignerSettings.bScaleChildrenWithParent || NewSize == FVector2D::ZeroVector)
-		{
-			return;
-		}
+	// Rotate to axis aligned, and apply the drag direction.
+	const FVector2D ResizeTranslation = MouseDownToCursorPositionGraphSpace.GetRotated(-OutputComponent->GetRotation()) * DragDirection;
+	const FVector2D RequestedSize = InitialSize + ResizeTranslation;
 
-		// Scale children with parent. 
-		FVector2D RatioVector = NewSize / OldSize;
-		for (UDMXPixelMappingBaseComponent* BaseChild : OutputComponent->GetChildren())
-		{
-			if (UDMXPixelMappingOutputComponent* Child = Cast<UDMXPixelMappingOutputComponent>(BaseChild))
-			{
-				Child->Modify();
-
-				// Scale size, at least 1x1 pixel
-				FVector2D NewChildSize = Child->GetSize() * RatioVector;
-				NewChildSize.X = FMath::Max(NewChildSize.X, 1.f);
-				NewChildSize.Y = FMath::Max(NewChildSize.Y, 1.f);
-				Child->SetSize(NewChildSize);
-
-				// Scale position
-				const FVector2D NewChildPosition = Child->GetPosition();
-				const FVector2D NewChildPositionRelative = (NewChildPosition - OutputComponent->GetPosition()) * RatioVector;
-				Child->SetPosition(OutputComponent->GetPosition() + NewChildPositionRelative);
-			}
-		}
+	const FVector2D NewSize = GetSnapSize(OutputComponent, RequestedSize);
+	const FVector2D OldSize = OutputComponent->GetSize();
+	if ((FMath::IsNearlyEqual(OldSize.X, NewSize.X) && FMath::IsNearlyEqual(OldSize.Y, NewSize.Y)) ||
+		NewSize.X < 1.f ||
+		NewSize.Y < 1.f)
+	{
+		// No (nearly) unchanged values, no size < 1 on either axis, no negative values.
+		return;
 	}
+	
+	OutputComponent->Modify();
+	OutputComponent->SetSize(NewSize);
 }
 
-FVector2D SDMXPixelMappingTransformHandle::GetSnapSize(UDMXPixelMappingOutputComponent* OutputComponent, const FVector2D& RequestedSize, const FVector2D& Direction) const
+FVector2D SDMXPixelMappingTransformHandle::GetSnapSize(UDMXPixelMappingOutputComponent* OutputComponent, const FVector2D& RequestedSize) const
 {
 	const TSharedPtr<FDMXPixelMappingToolkit> Toolkit = DesignerViewWeakPtr.Pin()->GetToolkit();
 	if (!Toolkit.IsValid() || !OutputComponent || !OutputComponent->GetRendererComponent())
@@ -271,34 +268,29 @@ FVector2D SDMXPixelMappingTransformHandle::GetSnapSize(UDMXPixelMappingOutputCom
 		}();
 
 	// Grid snap bottom right
-	const int32 BottomColumn = (OutputComponent->GetPosition().X + RequestedSize.X) / CellSize.X + 1;
-	const int32 RightRow = (OutputComponent->GetPosition().Y + RequestedSize.Y) / CellSize.Y + 1;
+	const FVector2D RequestedBottomRight = OutputComponent->GetPositionRotated() + RequestedSize;
+	const int32 BottomColumn = FMath::RoundHalfToZero(RequestedBottomRight.X / CellSize.X);
+	const int32 RightRow = FMath::RoundHalfToZero(RequestedBottomRight.Y / CellSize.Y);
 
-	FVector2D NewBottomRight = OutputComponent->GetPosition() + RequestedSize;
-	if (Direction.X < 0.f)
-	{
-		// Snap towards left
-		NewBottomRight.X = (BottomColumn - 1) * CellSize.X;
-	}
-	else if (Direction.X > 0.f)
+	FVector2D NewBottomRight = OutputComponent->GetPositionRotated() + RequestedSize;
+	if (DragDirection.X > 0.f)
 	{
 		// Snap towards right
 		NewBottomRight.X = BottomColumn * CellSize.X;
 	}
 
-	if (Direction.Y < 0.f)
-	{
-		// Snap towards top
-		NewBottomRight.Y = (RightRow - 1) * CellSize.Y;
-	}
-	else if (Direction.Y > 0.f)
+	if (DragDirection.Y > 0.f)
 	{
 		// Snap towards bottom
 		NewBottomRight.Y = RightRow * CellSize.Y;
 	}
 
-	const FVector2D SnapSize = NewBottomRight - OutputComponent->GetPosition();
-	return SnapSize;
+	const FVector2D SnapSize = NewBottomRight - OutputComponent->GetPositionRotated();
+	FVector2D SnapSizeClamped;
+	SnapSizeClamped.X = FMath::Max(SnapSize.X, 1.f);
+	SnapSizeClamped.Y = FMath::Max(SnapSize.Y, 1.f);
+
+	return SnapSizeClamped;
 }
 
 FVector2D SDMXPixelMappingTransformHandle::ComputeDragDirection(EDMXPixelMappingTransformDirection InTransformDirection) const
@@ -307,9 +299,6 @@ FVector2D SDMXPixelMappingTransformHandle::ComputeDragDirection(EDMXPixelMapping
 	{
 	case EDMXPixelMappingTransformDirection::CenterRight:
 		return FVector2D(1, 0);
-
-	case EDMXPixelMappingTransformDirection::BottomLeft:
-		return FVector2D(-1, 1);
 	case EDMXPixelMappingTransformDirection::BottomCenter:
 		return FVector2D(0, 1);
 	case EDMXPixelMappingTransformDirection::BottomRight:
@@ -327,9 +316,6 @@ FVector2D SDMXPixelMappingTransformHandle::ComputeOrigin(EDMXPixelMappingTransfo
 	{
 	case EDMXPixelMappingTransformDirection::CenterRight:
 		return Size * FVector2D(0, 0.5);
-
-	case EDMXPixelMappingTransformDirection::BottomLeft:
-		return Size * FVector2D(1, 0);
 	case EDMXPixelMappingTransformDirection::BottomCenter:
 		return Size * FVector2D(0.5, 0);
 	case EDMXPixelMappingTransformDirection::BottomRight:
diff --git a/Engine/Plugins/VirtualProduction/DMX/DMXPixelMapping/Source/DMXPixelMappingEditor/Private/Widgets/SDMXPixelMappingTransformHandle.h b/Engine/Plugins/VirtualProduction/DMX/DMXPixelMapping/Source/DMXPixelMappingEditor/Private/Widgets/SDMXPixelMappingTransformHandle.h
--- a/Engine/Plugins/VirtualProduction/DMX/DMXPixelMapping/Source/DMXPixelMappingEditor/Private/Widgets/SDMXPixelMappingTransformHandle.h
+++ b/Engine/Plugins/VirtualProduction/DMX/DMXPixelMapping/Source/DMXPixelMappingEditor/Private/Widgets/SDMXPixelMappingTransformHandle.h
@@ -14,7 +14,6 @@ enum class EDMXPixelMappingTransformDirection : uint8
 {
 	CenterRight,
 	BottomRight,
-	BottomLeft,
 	BottomCenter,
 
 	MAX
@@ -58,14 +57,14 @@ public:
 private:
 	EVisibility GetHandleVisibility() const;
 
-	/** Resizes the component on the next tick */
-	void RequestResize(UDMXPixelMappingBaseComponent* BaseComponent, const FVector2D& Direction);
-	
+	/** Rotate the component */
+	void Rotate(UDMXPixelMappingOutputComponent* OutputComponent, const FVector2D& CursorPosition, double WidgetScale);
+
 	/** Resizes the component */
-	void Resize(UDMXPixelMappingBaseComponent* BaseComponent, const FVector2D& Direction, const FVector2D& Amount);
+	void Resize(UDMXPixelMappingOutputComponent* OutputComponent, const FVector2D& CursorPosition, double WidgetScale);
 
 	/** Utiltity that returns the snap size of an output component. Note, this snaps to pixels if grid snapping is disabled. */
-	FVector2D GetSnapSize(UDMXPixelMappingOutputComponent* OutputComponent, const FVector2D& RequestedSize, const FVector2D& Direction) const;
+	FVector2D GetSnapSize(UDMXPixelMappingOutputComponent* OutputComponent, const FVector2D& RequestedSize) const;
 
 	EDMXPixelMappingTransformAction ComputeActionAtLocation(const FGeometry& MyGeometry, const FPointerEvent& MouseEvent) const;
 
@@ -75,11 +74,12 @@ private:
 	EDMXPixelMappingTransformDirection TransformDirection;
 	EDMXPixelMappingTransformAction Action;
 
-	FVector2D DragDirection;
-	FVector2D DragOrigin;
+	FVector2D DragDirection = FVector2D::ZeroVector;
+	FVector2D DragOrigin = FVector2D::ZeroVector;
 
-	FVector2D MouseDownPosition;
-	FMargin StartingOffsets;
+	FVector2D MouseDownPosition = FVector2D::ZeroVector;
+	FVector2D InitialSize = FVector2D::ZeroVector;
+	double InitialRotation = 0.0;
 
 	TSharedPtr<FScopedTransaction> ScopedTransaction;
 
@@ -88,5 +88,5 @@ private:
 	TAttribute<FVector2D> Offset;
 
 	/** Timer handle used when resize is requested */
-	FTimerHandle RequestResizeHandle;
+	FTimerHandle RequestApplyTransformHandle;
 };
diff --git a/Engine/Plugins/VirtualProduction/DMX/DMXPixelMapping/Source/DMXPixelMappingRuntime/Internal/Components/DMXPixelMappingComponentGeometryCache.h b/Engine/Plugins/VirtualProduction/DMX/DMXPixelMapping/Source/DMXPixelMappingRuntime/Internal/Components/DMXPixelMappingComponentGeometryCache.h
new file mode 100644
--- /dev/null
+++ b/Engine/Plugins/VirtualProduction/DMX/DMXPixelMapping/Source/DMXPixelMappingRuntime/Internal/Components/DMXPixelMappingComponentGeometryCache.h
@@ -0,0 +1,87 @@
+// Copyright Epic Games, Inc. All Rights Reserved.
+
+#pragma once
+
+
+#include "DMXPixelMappingComponentGeometryCache.generated.h"
+
+class UDMXPixelMappingOutputComponent;
+
+
+/** 
+ * Caches pixel mapping component geometry. 
+ * Useful for performance and to propagnonate changes incrementally to children.
+ * 
+ * The outer îs expected to be a DMXPixelMappingOutputComponent or a derived class.
+ */
+USTRUCT()
+struct FDMXPixelMappingComponentGeometryCache
+{
+	GENERATED_BODY()
+
+public:
+	/** Initializes the geometry cache */
+	void Initialize(UDMXPixelMappingOutputComponent* InOwner, const FVector2D& InPosition, const FVector2D& InSize, double InRotation);
+
+	/** Sets the absoulte position, returns the translation between new and previous position. */
+	FVector2D SetPositionAbsolute(const FVector2D& NewPosition);
+	const FVector2D& GetPositionAbsolute() const { return Position; }
+
+	/** Sets the absoulte position, rotated. Returns the translation between new and previous position. */
+	FVector2D SetPositionRotatedAbsolute(const FVector2D& NewPositionRotated);
+	FVector2D GetPositionRotatedAbsolute() const;
+
+	/** 
+	 * Sets the absoulte size, outputs the delta size and delta position. 
+	 * Delta position occurs if the rectangle is rotated, as its pivot is in the center, not the top left edge.
+	 * To keep it in rotated position, the position needs be updated.
+	 */
+	void SetSizeAbsolute(const FVector2D& InNewSize, FVector2D& OutDeltaSize, FVector2D& OutDeltaPosition);
+	
+	/** Sets the absolute size without outputting deltas */
+	void SetSizeAbsolute(const FVector2D& InNewSize);
+	const FVector2D& GetSizeAbsolute() const { return Size; }
+
+	/** Sets the absoulte rotation, returns the delta to the previous rotation. */
+	double SetRotationAbsolute(double NewRotation);
+	double GetRotationAbsolute() const { return Rotation;  }
+
+	/** Returns rotation as sin cos */
+	void GetSinCos(double& OutSin, double& OutCos) const;
+
+	/** Returns the edges of the rectangle, considering rotation. */
+	void GetEdges(FVector2D& A, FVector2D& B, FVector2D& C, FVector2D& D) const;
+
+	/** Propagonates position changes to children. */
+	void PropagonatePositionChangesToChildren(const FVector2D& Translation);
+
+	/** Propagonates size changes to children. */
+	void PropagonateSizeChangesToChildren(const FVector2D& DeltaSize, const FVector2D& DeltaPosition);
+
+	/** Propagonates rotation changes to children. */
+	void PropagonateRotationChangesToChildren(double DeltaRotation);
+
+private:
+	/** Position before it was changed */
+	UPROPERTY(Transient)
+	FVector2D Position = FVector2D::ZeroVector;
+
+	/** Size before it was changed */
+	UPROPERTY(Transient)
+	FVector2D Size = FVector2D::ZeroVector;
+
+	/** Rotation before it was changed */
+	UPROPERTY(Transient)
+	double Rotation = 0.f;
+
+	/** Cached Sine of Rotation. Preferable over non-reflected FQuat2D etc. for undo/redo purpose. */
+	UPROPERTY(Transient)
+	double Sin = 0.f;
+
+	/** Cached Cosine of Rotation. Preferable over non-reflected FQuat2D etc. for undo/redo purpose. */ 
+	UPROPERTY(Transient)
+	double Cos = 1.f;
+
+	/** The owner of this cache */
+	TWeakObjectPtr<UDMXPixelMappingOutputComponent> WeakOwner;
+};
diff --git a/Engine/Plugins/VirtualProduction/DMX/DMXPixelMapping/Source/DMXPixelMappingRuntime/Private/Components/DMXPixelMappingComponentGeometryCache.cpp b/Engine/Plugins/VirtualProduction/DMX/DMXPixelMapping/Source/DMXPixelMappingRuntime/Private/Components/DMXPixelMappingComponentGeometryCache.cpp
new file mode 100644
--- /dev/null
+++ b/Engine/Plugins/VirtualProduction/DMX/DMXPixelMapping/Source/DMXPixelMappingRuntime/Private/Components/DMXPixelMappingComponentGeometryCache.cpp
@@ -0,0 +1,171 @@
+// Copyright Epic Games, Inc. All Rights Reserved.
+
+#include "Components/DMXPixelMappingComponentGeometryCache.h"
+
+#include "Components/DMXPixelMappingOutputComponent.h"
+
+
+namespace UE::DMX::PixelMapping::ComponentGeometryCache::Private
+{
+	FVector2D RotateVector(const FVector2D& Vector, double Sin, double Cos, const FVector2D& Pivot)
+	{
+		const FVector2d PivotToVector = Vector - Pivot;
+		const FVector2d PivotToVectorRotated = FVector2D(
+			Cos * PivotToVector.X - Sin * PivotToVector.Y,
+			Sin * PivotToVector.X + Cos * PivotToVector.Y);
+
+		return Pivot + PivotToVectorRotated;
+	}
+}
+
+void FDMXPixelMappingComponentGeometryCache::Initialize(UDMXPixelMappingOutputComponent* InOwner, const FVector2D& InPosition, const FVector2D& InSize, double InRotation)
+{
+	WeakOwner = InOwner;
+	Position = InPosition;
+	Size = InSize;
+	Rotation = InRotation;
+}
+
+FVector2D FDMXPixelMappingComponentGeometryCache::SetPositionAbsolute(const FVector2D& NewPosition)
+{
+	const FVector2D Delta = NewPosition - Position;
+
+	Position = NewPosition;
+
+	return Delta;
+}
+
+FVector2D FDMXPixelMappingComponentGeometryCache::SetPositionRotatedAbsolute(const FVector2D& NewPositionRotated)
+{
+	const FVector2D OldPosition = Position;
+	const FVector2D Translation = NewPositionRotated - GetPositionRotatedAbsolute();
+
+	Position += Translation;
+
+	const FVector2D Delta = Position - OldPosition;
+	return Delta;
+}
+
+FVector2D FDMXPixelMappingComponentGeometryCache::GetPositionRotatedAbsolute() const
+{
+	const FVector2D Pivot = Position + Size / 2.0;
+
+	using namespace UE::DMX::PixelMapping::ComponentGeometryCache::Private;
+	return RotateVector(Position, Sin, Cos, Pivot);
+}
+
+void FDMXPixelMappingComponentGeometryCache::SetSizeAbsolute(const FVector2D& InNewSize, FVector2D& OutDeltaSize, FVector2D& OutDeltaPosition)
+{
+	const FVector2D Pivot = Position + Size / 2.0;
+
+	using namespace UE::DMX::PixelMapping::ComponentGeometryCache::Private;
+	OutDeltaPosition = RotateVector(Position, Sin, Cos, Position + Size / 2.f) - RotateVector(Position, Sin, Cos, Position + InNewSize / 2.f);
+	OutDeltaSize = InNewSize - Size;
+
+	Size = InNewSize;
+	Position += OutDeltaPosition;
+}
+
+void FDMXPixelMappingComponentGeometryCache::SetSizeAbsolute(const FVector2D& InNewSize)
+{
+	using namespace UE::DMX::PixelMapping::ComponentGeometryCache::Private;
+	const FVector2D DeltaPosition = RotateVector(Position, Sin, Cos, Position + Size / 2.f) - RotateVector(Position, Sin, Cos, Position + InNewSize / 2.f);
+
+	Size = InNewSize;
+	Position += DeltaPosition;
+}
+
+double FDMXPixelMappingComponentGeometryCache::SetRotationAbsolute(double NewRotation)
+{
+	const double Delta = NewRotation - Rotation;
+
+	Rotation = NewRotation;
+
+	// Cache rotation as sin cos
+	FMath::SinCos(&Sin, &Cos, FMath::DegreesToRadians(Rotation));
+
+	return Delta;
+}
+
+void FDMXPixelMappingComponentGeometryCache::GetSinCos(double& OutSin, double& OutCos) const
+{
+	OutSin = Sin;
+	OutCos = Cos;
+}
+
+void FDMXPixelMappingComponentGeometryCache::GetEdges(FVector2D& A, FVector2D& B, FVector2D& C, FVector2D& D) const
+{
+	const FVector2D Pivot = Position + Size / 2.0;
+
+	using namespace UE::DMX::PixelMapping::ComponentGeometryCache::Private;
+	A = RotateVector(Position, Sin, Cos, Pivot);
+	B = RotateVector(FVector2D(Position.X + Size.X, Position.Y), Sin, Cos, Pivot);
+	C = RotateVector(Position + Size, Sin, Cos, Pivot);
+	D = RotateVector(FVector2D(Position.X, Position.Y + Size.Y), Sin, Cos, Pivot);
+}
+
+void FDMXPixelMappingComponentGeometryCache::PropagonatePositionChangesToChildren(const FVector2D& Translation)
+{
+	if (UDMXPixelMappingOutputComponent* OutputComponent = Cast<UDMXPixelMappingOutputComponent>(WeakOwner.Get()))
+	{
+		constexpr bool bRecursive = false;
+		OutputComponent->ForEachChildOfClass<UDMXPixelMappingOutputComponent>(
+			[Translation](UDMXPixelMappingOutputComponent* Child)
+			{
+				Child->Modify();
+				Child->SetPosition(Child->GetPosition() + Translation);
+			},
+			bRecursive);
+	}
+}
+
+void FDMXPixelMappingComponentGeometryCache::PropagonateSizeChangesToChildren(const FVector2D& DeltaSize, const FVector2D& DeltaPosition)
+{
+	if (UDMXPixelMappingOutputComponent* OutputComponent = Cast<UDMXPixelMappingOutputComponent>(WeakOwner.Get()))
+	{
+		const double RestoreRotation = Rotation;
+		OutputComponent->SetRotation(0.0);
+
+		const FVector2D Scalar = Size / (Size - DeltaSize);
+
+		constexpr bool bRecursive = false;
+		OutputComponent->ForEachChildOfClass<UDMXPixelMappingOutputComponent>(
+			[Scalar, DeltaSize, this](UDMXPixelMappingOutputComponent* Child)
+			{
+				const FVector2D OldChildPosition = Child->GetPosition();
+				const FVector2D OldChildSize = Child->GetSize();
+
+				Child->Modify();
+				Child->SetSize(Child->GetSize() * Scalar);
+
+				// Reposition after setting the size, as set size and changes the center.
+				const FVector2D OldChildCenterRelative = OldChildPosition + OldChildSize / 2.0 - Position;
+				const FVector2D NewChildCenterRelative = OldChildCenterRelative * Scalar;
+				Child->SetPosition(Position + NewChildCenterRelative - Child->GetSize() / 2.0);
+			},
+			bRecursive);
+
+		OutputComponent->SetRotation(RestoreRotation);
+	}
+}
+
+void FDMXPixelMappingComponentGeometryCache::PropagonateRotationChangesToChildren(double DeltaRotation)
+{
+	if (UDMXPixelMappingOutputComponent* OutputComponent = Cast<UDMXPixelMappingOutputComponent>(WeakOwner.Get()))
+	{
+		constexpr bool bRecursive = false;
+		OutputComponent->ForEachChildOfClass<UDMXPixelMappingOutputComponent>(
+			[DeltaRotation, this](UDMXPixelMappingOutputComponent* Child)
+			{
+				Child->Modify();
+				Child->SetRotation(Child->GetRotation() + DeltaRotation);
+
+				// Rotate center around parent pivot
+				const FVector2D ParentPivotAbsolute = Position + Size / 2.0;
+				const FVector2D OldCenterRelative = Child->GetPosition() - ParentPivotAbsolute + Child->GetSize() / 2.0;
+				const FVector2D NewCenterRelative = OldCenterRelative.GetRotated(DeltaRotation);
+				Child->SetPosition(ParentPivotAbsolute + NewCenterRelative - Child->GetSize() / 2.0);
+			},
+			bRecursive);
+	}
+}
diff --git a/Engine/Plugins/VirtualProduction/DMX/DMXPixelMapping/Source/DMXPixelMappingRuntime/Private/Components/DMXPixelMappingFixtureGroupComponent.cpp b/Engine/Plugins/VirtualProduction/DMX/DMXPixelMapping/Source/DMXPixelMappingRuntime/Private/Components/DMXPixelMappingFixtureGroupComponent.cpp
--- a/Engine/Plugins/VirtualProduction/DMX/DMXPixelMapping/Source/DMXPixelMappingRuntime/Private/Components/DMXPixelMappingFixtureGroupComponent.cpp
+++ b/Engine/Plugins/VirtualProduction/DMX/DMXPixelMapping/Source/DMXPixelMappingRuntime/Private/Components/DMXPixelMappingFixtureGroupComponent.cpp
@@ -2,34 +2,19 @@
 
 #include "Components/DMXPixelMappingFixtureGroupComponent.h"
 
+#include "Components/DMXPixelMappingComponentGeometryCache.h"
 #include "Components/DMXPixelMappingMatrixComponent.h"
 #include "Components/DMXPixelMappingRendererComponent.h"
+#include "DMXPixelMapping.h"
 #include "Library/DMXLibrary.h"
 
-#if WITH_EDITOR
-#include "DMXPixelMappingComponentWidget.h"
-#endif // WITH_EDITOR
-
 
 #define LOCTEXT_NAMESPACE "DMXPixelMappingFixtureGroupComponent"
 
 UDMXPixelMappingFixtureGroupComponent::UDMXPixelMappingFixtureGroupComponent()
 {
-	SetSize(FVector2D(500.f, 500.f));
-}
-
-void UDMXPixelMappingFixtureGroupComponent::PostInitProperties()
-{
-	Super::PostInitProperties();
-
-	LastPosition = GetPosition();
-}
-
-void UDMXPixelMappingFixtureGroupComponent::PostLoad()
-{
-	Super::PostLoad();
-
-	LastPosition = GetPosition();
+	SizeX = 500.f;
+	SizeY = 500.f;
 }
 
 #if WITH_EDITOR
@@ -39,34 +24,11 @@ void UDMXPixelMappingFixtureGroupComponent::PostEditChangeProperty(FPropertyChan
 	Super::PostEditChangeProperty(PropertyChangedEvent);
 
 	const FName PropertyName = PropertyChangedEvent.GetPropertyName();
-	if (PropertyName == UDMXPixelMappingOutputComponent::GetPositionXPropertyName() ||
-		PropertyName == UDMXPixelMappingOutputComponent::GetPositionYPropertyName())
-	{
-		HandlePositionChanged();
-	}
-	else if (PropertyName == GET_MEMBER_NAME_CHECKED(UDMXPixelMappingFixtureGroupComponent, DMXLibrary))
+	if (PropertyName == GET_MEMBER_NAME_CHECKED(UDMXPixelMappingFixtureGroupComponent, DMXLibrary))
 	{
 		OnDMXLibraryChangedDelegate.Broadcast();
 	}
 
-	PRAGMA_DISABLE_DEPRECATION_WARNINGS
-	if (PropertyName == UDMXPixelMappingOutputComponent::GetSizeXPropertyName() ||
-		PropertyName == UDMXPixelMappingOutputComponent::GetSizeYPropertyName())
-	{
-		if (ComponentWidget_DEPRECATED.IsValid())
-		{
-			ComponentWidget_DEPRECATED->SetSize(GetSize());
-		}
-	}
-	else if (PropertyName == GET_MEMBER_NAME_CHECKED(UDMXPixelMappingFixtureGroupComponent, DMXLibrary))
-	{
-		if (ComponentWidget_DEPRECATED.IsValid())
-		{
-			ComponentWidget_DEPRECATED->SetLabelText(FText::FromString(GetUserFriendlyName()));
-		}
-	}
-	PRAGMA_ENABLE_DEPRECATION_WARNINGS
-
 	InvalidatePixelMapRenderer();
 }
 #endif // WITH_EDITOR
@@ -76,9 +38,6 @@ void UDMXPixelMappingFixtureGroupComponent::PostEditUndo()
 {
 	Super::PostEditUndo();
 
-	// Update last position, so the next will be set correctly on children
-	LastPosition = GetPosition();
-
 	// Restore Matrices 
 	for (UDMXPixelMappingBaseComponent* Child : Children)
 	{
@@ -155,61 +114,79 @@ void UDMXPixelMappingFixtureGroupComponent::QueueDownsample()
 	PRAGMA_ENABLE_DEPRECATION_WARNINGS
 }
 
+FString UDMXPixelMappingFixtureGroupComponent::GetUserName() const
+{
+	if (UserName.IsEmpty())
+	{
+		const FString DMXLibraryName = DMXLibrary ? DMXLibrary->GetName() : LOCTEXT("NoDMXLibrary", "No DMX Library").ToString();
+		return FString::Printf(TEXT("Fixture Group: %s"), *DMXLibraryName);
+	}
+	else
+	{
+		return UserName;
+	}
+}
+
 void UDMXPixelMappingFixtureGroupComponent::SetPosition(const FVector2D& NewPosition)
 {
-	Super::SetPosition(NewPosition);
+	PositionX = NewPosition.X;
+	PositionY = NewPosition.Y;
+
+	const FVector2D Translation = CachedGeometry.SetPositionAbsolute(NewPosition);
+	CachedGeometry.PropagonatePositionChangesToChildren(Translation);
 
-	HandlePositionChanged();
+#if WITH_EDITOR
+	EditorPositionWithRotation = CachedGeometry.GetPositionRotatedAbsolute();
+#endif
 }
 
-void UDMXPixelMappingFixtureGroupComponent::SetSize(const FVector2D& NewSize)
+void UDMXPixelMappingFixtureGroupComponent::SetPositionRotated(FVector2D NewRotatedPosition)
 {
-	Super::SetSize(NewSize);
+	const FVector2D Translation = CachedGeometry.SetPositionRotatedAbsolute(NewRotatedPosition);
+
+	PositionX = CachedGeometry.GetPositionAbsolute().X;
+	PositionY = CachedGeometry.GetPositionAbsolute().Y;
+
+	CachedGeometry.PropagonatePositionChangesToChildren(Translation);
 
 #if WITH_EDITOR
-	PRAGMA_DISABLE_DEPRECATION_WARNINGS
-	if (ComponentWidget_DEPRECATED.IsValid())
-	{
-		ComponentWidget_DEPRECATED->SetSize(GetSize());
-	}
-	PRAGMA_ENABLE_DEPRECATION_WARNINGS
+	EditorPositionWithRotation = CachedGeometry.GetPositionRotatedAbsolute();
 #endif
 }
 
-void UDMXPixelMappingFixtureGroupComponent::HandlePositionChanged()
+void UDMXPixelMappingFixtureGroupComponent::SetSize(const FVector2D& NewSize)
 {
+	SizeX = NewSize.X;
+	SizeY = NewSize.Y;
+
+	FVector2D DeltaSize;
+	FVector2D DeltaPosition;
+	CachedGeometry.SetSizeAbsolute(NewSize, DeltaSize, DeltaPosition);
+
+	// Resizing rotated components moves their pivot, and by that their position without rotation.
+	// Adjust for this so the component keeps its rotated position.
+	SetPosition(CachedGeometry.GetPositionAbsolute());
+
+	// Propagonate size changes to children if desired
 #if WITH_EDITOR
-	PRAGMA_DISABLE_DEPRECATION_WARNINGS
-	if (ComponentWidget_DEPRECATED.IsValid())
+	UDMXPixelMapping* PixelMapping = GetPixelMapping();
+	if (PixelMapping && PixelMapping->bEditorScaleChildrenWithParent)
 	{
-		ComponentWidget_DEPRECATED->SetPosition(GetPosition());
+		CachedGeometry.PropagonateSizeChangesToChildren(DeltaSize, DeltaPosition);
 	}
-	PRAGMA_ENABLE_DEPRECATION_WARNINGS
 #endif
-
-	// Propagonate to children
-	constexpr bool bUpdatePositionRecursive = false;
-	ForEachChildOfClass<UDMXPixelMappingOutputComponent>([this](UDMXPixelMappingOutputComponent* ChildComponent)
-		{
-			const FVector2D ChildOffset = ChildComponent->GetPosition() - LastPosition;
-			ChildComponent->SetPosition(GetPosition() + ChildOffset);
-
-		}, bUpdatePositionRecursive);
-
-	LastPosition = GetPosition();
 }
 
-FString UDMXPixelMappingFixtureGroupComponent::GetUserName() const
+void UDMXPixelMappingFixtureGroupComponent::SetRotation(double NewRotation)
 {
-	if (UserName.IsEmpty())
-	{
-		const FString DMXLibraryName = DMXLibrary ? DMXLibrary->GetName() : LOCTEXT("NoDMXLibrary", "No DMX Library").ToString();
-		return FString::Printf(TEXT("Fixture Group: %s"), *DMXLibraryName);
-	}
-	else
-	{
-		return UserName;
-	}
+	Rotation = NewRotation;
+
+	const double DeltaRotation = CachedGeometry.SetRotationAbsolute(NewRotation);
+	CachedGeometry.PropagonateRotationChangesToChildren(DeltaRotation);
+
+#if WITH_EDITOR
+		EditorPositionWithRotation = CachedGeometry.GetPositionRotatedAbsolute();
+#endif
 }
 
 #if WITH_EDITOR
diff --git a/Engine/Plugins/VirtualProduction/DMX/DMXPixelMapping/Source/DMXPixelMappingRuntime/Private/Components/DMXPixelMappingFixtureGroupItemComponent.cpp b/Engine/Plugins/VirtualProduction/DMX/DMXPixelMapping/Source/DMXPixelMappingRuntime/Private/Components/DMXPixelMappingFixtureGroupItemComponent.cpp
--- a/Engine/Plugins/VirtualProduction/DMX/DMXPixelMapping/Source/DMXPixelMappingRuntime/Private/Components/DMXPixelMappingFixtureGroupItemComponent.cpp
+++ b/Engine/Plugins/VirtualProduction/DMX/DMXPixelMapping/Source/DMXPixelMappingRuntime/Private/Components/DMXPixelMappingFixtureGroupItemComponent.cpp
@@ -2,25 +2,21 @@
 
 #include "Components/DMXPixelMappingFixtureGroupItemComponent.h"
 
+#include "ColorSpace/DMXPixelMappingColorSpace_RGBCMY.h"
+#include "Components/DMXPixelMappingComponentGeometryCache.h"
+#include "Components/DMXPixelMappingFixtureGroupComponent.h"
+#include "Components/DMXPixelMappingRendererComponent.h"
 #include "DMXConversions.h"
 #include "DMXPixelMappingMainStreamObjectVersion.h"
 #include "DMXPixelMappingRuntimeUtils.h"
 #include "DMXPixelMappingTypes.h"
-#include "ColorSpace/DMXPixelMappingColorSpace_RGBCMY.h"
-#include "Components/DMXPixelMappingFixtureGroupComponent.h"
-#include "Components/DMXPixelMappingRendererComponent.h"
+#include "Engine/Texture.h"
 #include "Library/DMXEntityFixturePatch.h"
 #include "Library/DMXEntityFixtureType.h"
 #include "Library/DMXLibrary.h"
 #include "IO/DMXOutputPort.h"
 #include "Modulators/DMXModulator.h"
-
-#if WITH_EDITOR
-#include "DMXPixelMappingComponentWidget.h"
-#endif // WITH_EDITOR
-
 #include "TextureResource.h"
-#include "Engine/Texture.h"
 
 
 DECLARE_CYCLE_STAT(TEXT("Send Fixture Group Item"), STAT_DMXPixelMaping_FixtureGroupItem, STATGROUP_DMXPIXELMAPPING);
@@ -33,7 +29,8 @@ UDMXPixelMappingFixtureGroupItemComponent::UDMXPixelMappingFixtureGroupItemCompo
 	ColorSpaceClass = UDMXPixelMappingColorSpace_RGBCMY::StaticClass();
 	ColorSpace = CreateDefaultSubobject<UDMXPixelMappingColorSpace_RGBCMY>("ColorSpace");
 
-	SetSize(FVector2D(32.f, 32.f));
+	SizeX = 32.f;
+	SizeY = 32.f;
 
 #if WITH_EDITORONLY_DATA
 	// Even tho deprecated, default values on deprecated properties need be set so they don't load their type's default value.
@@ -180,25 +177,6 @@ void UDMXPixelMappingFixtureGroupItemComponent::PostEditChangeProperty(FProperty
 		}
 	}
 
-	PRAGMA_DISABLE_DEPRECATION_WARNINGS
-	if (PropertyName == UDMXPixelMappingOutputComponent::GetPositionXPropertyName() ||
-		PropertyName == UDMXPixelMappingOutputComponent::GetPositionYPropertyName())
-	{
-		if (ComponentWidget_DEPRECATED.IsValid())
-		{
-			ComponentWidget_DEPRECATED->SetPosition(GetPosition());
-		}
-	}
-	else if (PropertyName == UDMXPixelMappingOutputComponent::GetSizeXPropertyName() ||
-		PropertyName == UDMXPixelMappingOutputComponent::GetSizeYPropertyName())
-	{
-		if (ComponentWidget_DEPRECATED.IsValid())
-		{
-			ComponentWidget_DEPRECATED->SetSize(GetSize());
-		}
-	}
-	PRAGMA_ENABLE_DEPRECATION_WARNINGS
-
 	InvalidatePixelMapRenderer();
 }
 #endif // WITH_EDITOR
@@ -380,15 +358,6 @@ void UDMXPixelMappingFixtureGroupItemComponent::SetPosition(const FVector2D& New
 {
 	Super::SetPosition(NewPosition);
 
-#if WITH_EDITOR
-	PRAGMA_DISABLE_DEPRECATION_WARNINGS
-	if (ComponentWidget_DEPRECATED.IsValid())
-	{
-		ComponentWidget_DEPRECATED->SetPosition(GetPosition());
-	}
-	PRAGMA_ENABLE_DEPRECATION_WARNINGS
-#endif
-
 	UpdateRenderElement();
 }
 
@@ -396,38 +365,24 @@ void UDMXPixelMappingFixtureGroupItemComponent::SetSize(const FVector2D& NewSize
 {
 	Super::SetSize(NewSize);
 
-#if WITH_EDITOR
-	PRAGMA_DISABLE_DEPRECATION_WARNINGS
-	if (ComponentWidget_DEPRECATED.IsValid())
-	{
-		ComponentWidget_DEPRECATED->SetSize(GetSize());
-	}
-	PRAGMA_ENABLE_DEPRECATION_WARNINGS
-#endif
-
 	UpdateRenderElement();
 }
 
 bool UDMXPixelMappingFixtureGroupItemComponent::IsOverParent() const
 {
-	// Needs be over the over the group
-	if (UDMXPixelMappingFixtureGroupComponent* ParentFixtureGroupComponent = Cast<UDMXPixelMappingFixtureGroupComponent>(GetParent()))
+	if (UDMXPixelMappingFixtureGroupComponent* Parent = Cast<UDMXPixelMappingFixtureGroupComponent>(GetParent()))
 	{
-		const float Left = GetPosition().X;
-		const float Top = GetPosition().Y;
-		const float Right = GetPosition().X + GetSize().X;
-		const float Bottom = GetPosition().Y + GetSize().Y;
-
-		const float ParentLeft = ParentFixtureGroupComponent->GetPosition().X;
-		const float ParentTop = ParentFixtureGroupComponent->GetPosition().Y;
-		const float ParentRight = ParentFixtureGroupComponent->GetPosition().X + ParentFixtureGroupComponent->GetSize().X;
-		const float ParentBottom = ParentFixtureGroupComponent->GetPosition().Y + ParentFixtureGroupComponent->GetSize().Y;
+		FVector2D A;
+		FVector2D B;
+		FVector2D C;
+		FVector2D D;
+		CachedGeometry.GetEdges(A, B, C, D);
 
 		return
-			Left > ParentLeft - .49f &&
-			Top > ParentTop - .49f &&
-			Right < ParentRight + .49f &&
-			Bottom < ParentBottom + .49f;
+			Parent->IsOverPosition(A) &&
+			Parent->IsOverPosition(B) &&
+			Parent->IsOverPosition(C) &&
+			Parent->IsOverPosition(D);
 	}
 
 	return false;
diff --git a/Engine/Plugins/VirtualProduction/DMX/DMXPixelMapping/Source/DMXPixelMappingRuntime/Private/Components/DMXPixelMappingMatrixCellComponent.cpp b/Engine/Plugins/VirtualProduction/DMX/DMXPixelMapping/Source/DMXPixelMappingRuntime/Private/Components/DMXPixelMappingMatrixCellComponent.cpp
--- a/Engine/Plugins/VirtualProduction/DMX/DMXPixelMapping/Source/DMXPixelMappingRuntime/Private/Components/DMXPixelMappingMatrixCellComponent.cpp
+++ b/Engine/Plugins/VirtualProduction/DMX/DMXPixelMapping/Source/DMXPixelMappingRuntime/Private/Components/DMXPixelMappingMatrixCellComponent.cpp
@@ -3,6 +3,7 @@
 #include "Components/DMXPixelMappingMatrixCellComponent.h"
 
 #include "ColorSpace/DMXPixelMappingColorSpace.h"
+#include "Components/DMXPixelMappingComponentGeometryCache.h"
 #include "Components/DMXPixelMappingFixtureGroupComponent.h"
 #include "Components/DMXPixelMappingMatrixComponent.h"
 #include "Components/DMXPixelMappingRendererComponent.h"
@@ -14,16 +15,9 @@
 #include "Library/DMXEntityFixtureType.h"
 #include "TextureResource.h"
 
-#if WITH_EDITOR
-#include "DMXPixelMappingComponentWidget.h"
-#include "SDMXPixelMappingComponentBox.h"
-#include "SDMXPixelMappingComponentLabel.h"
-#endif // WITH_EDITOR
-
 
 #define LOCTEXT_NAMESPACE "DMXPixelMappingMatrixPixelComponent"
 
-
 UDMXPixelMappingMatrixCellComponent::UDMXPixelMappingMatrixCellComponent()
 	: DownsamplePixelIndex(0)
 {
@@ -74,151 +68,35 @@ void UDMXPixelMappingMatrixCellComponent::PostEditChangeProperty(FPropertyChange
 
 	UpdateRenderElement();
 
-	FName PropertyName = PropertyChangedEvent.GetPropertyName();
-
-	PRAGMA_DISABLE_DEPRECATION_WARNINGS
-	if (PropertyChangedEvent.GetPropertyName() == UDMXPixelMappingOutputComponent::GetPositionXPropertyName() ||
-		PropertyChangedEvent.GetPropertyName() == UDMXPixelMappingOutputComponent::GetPositionYPropertyName())
-	{
-		if (ComponentWidget_DEPRECATED.IsValid())
-		{
-			ComponentWidget_DEPRECATED->SetPosition(GetPosition());
-		}
-	}
-	if (PropertyChangedEvent.GetPropertyName() == UDMXPixelMappingOutputComponent::GetSizeXPropertyName() ||
-		PropertyChangedEvent.GetPropertyName() == UDMXPixelMappingOutputComponent::GetSizeYPropertyName())
-	{
-		if (ComponentWidget_DEPRECATED.IsValid())
-		{
-			ComponentWidget_DEPRECATED->SetSize(GetSize());
-		}
-	}
-	if (PropertyName == GET_MEMBER_NAME_CHECKED(UDMXPixelMappingMatrixCellComponent, CellID))
-	{
-		if (ComponentWidget_DEPRECATED.IsValid())
-		{
-			ComponentWidget_DEPRECATED->GetComponentBox()->SetIDText(FText::Format(LOCTEXT("CellID", "{0}"), CellID));
-		}
-	}
-	PRAGMA_ENABLE_DEPRECATION_WARNINGS
-
 	InvalidatePixelMapRenderer();
 }
 #endif // WITH_EDITOR
 
-#if WITH_EDITOR
-PRAGMA_DISABLE_DEPRECATION_WARNINGS
-TSharedRef<FDMXPixelMappingComponentWidget> UDMXPixelMappingMatrixCellComponent::BuildSlot(TSharedRef<SConstraintCanvas> InCanvas)
-{
-	ComponentWidget_DEPRECATED = Super::BuildSlot(InCanvas);
-
-	// Expect super to construct the component widget
-	if (ensureMsgf(ComponentWidget_DEPRECATED.IsValid(), TEXT("PixelMapping: Expected Super to construct a component widget, but didn't.")))
-	{
-		ComponentWidget_DEPRECATED->GetComponentLabel()->SetText(FText::GetEmpty());
-		ComponentWidget_DEPRECATED->GetComponentBox()->SetIDText(FText::Format(LOCTEXT("CellID", "{0}"), CellID));
-	}
-
-	return ComponentWidget_DEPRECATED.ToSharedRef();
-}
-PRAGMA_ENABLE_DEPRECATION_WARNINGS
-#endif // WITH_EDITOR
-
-#if WITH_EDITOR
-bool UDMXPixelMappingMatrixCellComponent::IsVisible() const
-{
-	// Needs be over the matrix and over the group
-	if (UDMXPixelMappingMatrixComponent* ParentMatrixComponent = Cast<UDMXPixelMappingMatrixComponent>(GetParent()))
-	{
-		if (!ParentMatrixComponent->IsVisible())
-		{
-			return false;
-		}
-	}
-
-	return Super::IsVisible();
-}
-#endif // WITH_EDITOR
-
 #if WITH_EDITOR
 FLinearColor UDMXPixelMappingMatrixCellComponent::GetEditorColor() const
 {
-	PRAGMA_DISABLE_DEPRECATION_WARNINGS
-	if (bLockInDesigner)
-	{
-		// When locked in designer, always use the parent color. So when the parent shows an error color, show it too.
-		if (UDMXPixelMappingMatrixComponent* ParentMatrixComponent = Cast< UDMXPixelMappingMatrixComponent>(GetParent()))
-		{
-			if (const TSharedPtr<FDMXPixelMappingComponentWidget>& ParentComponentWidget = ParentMatrixComponent->GetComponentWidget())
-			{
-				return ParentComponentWidget->GetColor();
-			}
-		}
-	}
-	PRAGMA_ENABLE_DEPRECATION_WARNINGS
-
 	return EditorColor;
 }
 #endif // WITH_EDITOR
 
 bool UDMXPixelMappingMatrixCellComponent::IsOverParent() const
 {
-	// Needs be over the matrix and over the group
-	UDMXPixelMappingMatrixComponent* ParentMatrixComponent = Cast<UDMXPixelMappingMatrixComponent>(GetParent());
-	if (!ParentMatrixComponent)
+	if (UDMXPixelMappingMatrixComponent* Parent = Cast<UDMXPixelMappingMatrixComponent>(GetParent()))
 	{
-		return false;
-	}
-	
-	const bool bIsParentMatrixOverGroup = ParentMatrixComponent->IsOverParent();
-	if (!bIsParentMatrixOverGroup)
-	{
-		return false;
-	}
+		FVector2D A;
+		FVector2D B;
+		FVector2D C;
+		FVector2D D;
+		CachedGeometry.GetEdges(A, B, C, D);
 
-	const float Left = GetPosition().X;
-	const float Top = GetPosition().Y;
-	const float Right = GetPosition().X + GetSize().X;
-	const float Bottom = GetPosition().Y + GetSize().Y;
-
-	const float ParentLeft = ParentMatrixComponent->GetPosition().X;
-	const float ParentTop = ParentMatrixComponent->GetPosition().Y;
-	const float ParentRight = ParentMatrixComponent->GetPosition().X + ParentMatrixComponent->GetSize().X;
-	const float ParentBottom = ParentMatrixComponent->GetPosition().Y + ParentMatrixComponent->GetSize().Y;
-
-	return
-		Left > ParentLeft - .49f &&
-		Top > ParentTop - .49f &&
-		Right < ParentRight + .49f &&
-		Bottom < ParentBottom + .49f;
-}
-
-void UDMXPixelMappingMatrixCellComponent::SetPosition(const FVector2D& NewPosition)
-{
-	Super::SetPosition(NewPosition);
-
-#if WITH_EDITOR
-	PRAGMA_DISABLE_DEPRECATION_WARNINGS
-	if (ComponentWidget_DEPRECATED.IsValid())
-	{
-		ComponentWidget_DEPRECATED->SetPosition(GetPosition());
+		return
+			Parent->IsOverPosition(A) &&
+			Parent->IsOverPosition(B) &&
+			Parent->IsOverPosition(C) &&
+			Parent->IsOverPosition(D);
 	}
-	PRAGMA_ENABLE_DEPRECATION_WARNINGS
-#endif
-}
-
-void UDMXPixelMappingMatrixCellComponent::SetSize(const FVector2D& NewSize)
-{
-	Super::SetSize(NewSize);
 
-#if WITH_EDITOR
-	PRAGMA_DISABLE_DEPRECATION_WARNINGS
-	if (ComponentWidget_DEPRECATED.IsValid())
-	{
-		ComponentWidget_DEPRECATED->SetSize(GetSize());
-	}
-	PRAGMA_ENABLE_DEPRECATION_WARNINGS
-#endif
+	return false;
 }
 
 const FName& UDMXPixelMappingMatrixCellComponent::GetNamePrefix()
diff --git a/Engine/Plugins/VirtualProduction/DMX/DMXPixelMapping/Source/DMXPixelMappingRuntime/Private/Components/DMXPixelMappingMatrixComponent.cpp b/Engine/Plugins/VirtualProduction/DMX/DMXPixelMapping/Source/DMXPixelMappingRuntime/Private/Components/DMXPixelMappingMatrixComponent.cpp
--- a/Engine/Plugins/VirtualProduction/DMX/DMXPixelMapping/Source/DMXPixelMappingRuntime/Private/Components/DMXPixelMappingMatrixComponent.cpp
+++ b/Engine/Plugins/VirtualProduction/DMX/DMXPixelMapping/Source/DMXPixelMappingRuntime/Private/Components/DMXPixelMappingMatrixComponent.cpp
@@ -4,6 +4,7 @@
 
 #include "Algo/Sort.h"
 #include "ColorSpace/DMXPixelMappingColorSpace_RGBCMY.h"
+#include "Components/DMXPixelMappingComponentGeometryCache.h"
 #include "Components/DMXPixelMappingFixtureGroupComponent.h"
 #include "Components/DMXPixelMappingMatrixCellComponent.h"
 #include "DMXConversions.h"
@@ -22,10 +23,6 @@
 #include "Modulators/DMXModulator.h"
 #include "Templates/DMXPixelMappingComponentTemplate.h"
 
-#if WITH_EDITOR
-#include "DMXPixelMappingComponentWidget.h"
-#endif // WITH_EDITOR
-
 
 #define LOCTEXT_NAMESPACE "DMXPixelMappingMatrixComponent"
 
@@ -34,7 +31,8 @@ UDMXPixelMappingMatrixComponent::UDMXPixelMappingMatrixComponent()
 	ColorSpaceClass = UDMXPixelMappingColorSpace_RGBCMY::StaticClass();
 	ColorSpace = CreateDefaultSubobject<UDMXPixelMappingColorSpace_RGBCMY>("ColorSpace");
 
-	SetSize(FVector2D(32.f, 32.f));
+	SizeX = 32.f;
+	SizeY = 32.f;
 
 #if WITH_EDITORONLY_DATA
 	bExpanded = false;
@@ -97,6 +95,11 @@ void UDMXPixelMappingMatrixComponent::PostLoad()
 {
 	Super::PostLoad();
 
+	if (IsTemplate())
+	{
+		return;
+	}
+
 	// Add valid modulators to modulator classes, remove invalid modulators
 	for (int32 IndexModulator = 0; Modulators.IsValidIndex(IndexModulator); )
 	{
@@ -127,12 +130,14 @@ void UDMXPixelMappingMatrixComponent::PostInitProperties()
 {
 	Super::PostInitProperties();
 
-	if (!HasAnyFlags(RF_ClassDefaultObject | RF_ArchetypeObject))
+	if (IsTemplate())
 	{
-		// Listen to Fixture Type and Fixture Patch changes
-		UDMXEntityFixtureType::GetOnFixtureTypeChanged().AddUObject(this, &UDMXPixelMappingMatrixComponent::OnFixtureTypeChanged);
-		UDMXEntityFixturePatch::GetOnFixturePatchChanged().AddUObject(this, &UDMXPixelMappingMatrixComponent::OnFixturePatchChanged);
+		return;
 	}
+
+	// Listen to Fixture Type and Fixture Patch changes
+	UDMXEntityFixtureType::GetOnFixtureTypeChanged().AddUObject(this, &UDMXPixelMappingMatrixComponent::OnFixtureTypeChanged);
+	UDMXEntityFixturePatch::GetOnFixturePatchChanged().AddUObject(this, &UDMXPixelMappingMatrixComponent::OnFixturePatchChanged);
 }
 
 void UDMXPixelMappingMatrixComponent::LogInvalidProperties()
@@ -160,24 +165,6 @@ void UDMXPixelMappingMatrixComponent::LogInvalidProperties()
 	}
 }
 
-#if WITH_EDITOR
-void UDMXPixelMappingMatrixComponent::PreEditChange(FProperty* PropertyAboutToChange)
-{
-	Super::PreEditChange(PropertyAboutToChange);
-
-	if (PropertyAboutToChange)
-	{
-		const FName PropertyName = PropertyAboutToChange->GetFName();
-
-		if (PropertyName == UDMXPixelMappingOutputComponent::GetPositionXPropertyName() ||
-			PropertyName == UDMXPixelMappingOutputComponent::GetPositionYPropertyName())
-		{
-			PreEditChangePosition = GetPosition();
-		}
-	}
-}
-#endif // WITH_EDITOR
-
 #if WITH_EDITOR
 void UDMXPixelMappingMatrixComponent::PostEditChangeProperty(FPropertyChangedEvent& PropertyChangedEvent)
 {
@@ -187,7 +174,6 @@ void UDMXPixelMappingMatrixComponent::PostEditChangeProperty(FPropertyChangedEve
 	Super::PostEditChangeProperty(PropertyChangedEvent);
 
 	const FName PropertyName = PropertyChangedEvent.GetPropertyName();
-
 	if (PropertyName == GET_MEMBER_NAME_CHECKED(UDMXPixelMappingMatrixComponent, ColorSpaceClass))
 	{
 		if (ensureAlwaysMsgf(ColorSpaceClass, TEXT("Color Space Class was set to nullptr. This is not supported.")))
@@ -195,33 +181,7 @@ void UDMXPixelMappingMatrixComponent::PostEditChangeProperty(FPropertyChangedEve
 			ColorSpace = NewObject<UDMXPixelMappingColorSpace>(this, ColorSpaceClass);
 		}
 	}
-	
-	if (PropertyName == UDMXPixelMappingOutputComponent::GetPositionXPropertyName() ||
-		PropertyName == UDMXPixelMappingOutputComponent::GetPositionYPropertyName())
-	{
-		constexpr bool bUpdateSizeRecursive = false;
-		ForEachChildOfClass<UDMXPixelMappingMatrixCellComponent>([this](UDMXPixelMappingMatrixCellComponent* ChildComponent)
-			{
-				FVector2D RelativePosition = ChildComponent->GetPosition() - PreEditChangePosition;
-				ChildComponent->SetPosition(GetPosition() + RelativePosition);
-
-			}, bUpdateSizeRecursive);
-
-		HandlePositionChanged();
-	}
-	
-	if (PropertyName == UDMXPixelMappingOutputComponent::GetSizeXPropertyName() ||
-		PropertyName == UDMXPixelMappingOutputComponent::GetSizeYPropertyName() ||
-		PropertyName == GET_MEMBER_NAME_CHECKED(UDMXPixelMappingMatrixComponent, CellSize))
-	{
-		HandleSizeChanged();
-
-		// Update size again from the new CellSize that results from handling size or matrix changes
-		const FVector2D NewSize(CellSize.X * CoordinateGrid.X, CellSize.Y * CoordinateGrid.Y);
-		SetSize(NewSize);
-	}
-
-	if (PropertyName == GET_MEMBER_NAME_CHECKED(UDMXPixelMappingMatrixComponent, FixturePatchRef) ||
+	else if (PropertyName == GET_MEMBER_NAME_CHECKED(UDMXPixelMappingMatrixComponent, FixturePatchRef) ||
 		PropertyName == GET_MEMBER_NAME_CHECKED(UDMXPixelMappingMatrixComponent, CoordinateGrid) ||
 		PropertyName == GET_MEMBER_NAME_CHECKED(FDMXEntityReference, DMXLibrary))
 	{
@@ -251,24 +211,19 @@ const FText UDMXPixelMappingMatrixComponent::GetPaletteCategory()
 
 bool UDMXPixelMappingMatrixComponent::IsOverParent() const
 {
-	// Needs be over the over the group
-	if (UDMXPixelMappingFixtureGroupComponent* ParentFixtureGroupComponent = Cast<UDMXPixelMappingFixtureGroupComponent>(GetParent()))
+	if (UDMXPixelMappingFixtureGroupComponent* Parent = Cast<UDMXPixelMappingFixtureGroupComponent>(GetParent()))
 	{
-		const float Left = GetPosition().X;
-		const float Top = GetPosition().Y;
-		const float Right = GetPosition().X + GetSize().X;
-		const float Bottom = GetPosition().Y + GetSize().Y;
-
-		const float ParentLeft = ParentFixtureGroupComponent->GetPosition().X;
-		const float ParentTop = ParentFixtureGroupComponent->GetPosition().Y;
-		const float ParentRight = ParentFixtureGroupComponent->GetPosition().X + ParentFixtureGroupComponent->GetSize().X;
-		const float ParentBottom = ParentFixtureGroupComponent->GetPosition().Y + ParentFixtureGroupComponent->GetSize().Y;
+		FVector2D A;
+		FVector2D B;
+		FVector2D C;
+		FVector2D D;
+		CachedGeometry.GetEdges(A, B, C, D);
 
 		return
-			Left > ParentLeft - .49f &&
-			Top > ParentTop - .49f &&
-			Right < ParentRight + .49f &&
-			Bottom < ParentBottom + .49f;
+			Parent->IsOverPosition(A) &&
+			Parent->IsOverPosition(B) &&
+			Parent->IsOverPosition(C) &&
+			Parent->IsOverPosition(D);
 	}
 
 	return false;
@@ -418,86 +373,67 @@ FString UDMXPixelMappingMatrixComponent::GetUserName() const
 
 void UDMXPixelMappingMatrixComponent::SetPosition(const FVector2D& NewPosition)
 {
-	const FVector2D OldPosition = GetPosition();
-
-	Super::SetPosition(NewPosition);
-
-	constexpr bool bUpdateSizeRecursive = false;
-	ForEachChildOfClass<UDMXPixelMappingMatrixCellComponent>([&OldPosition, &NewPosition](UDMXPixelMappingMatrixCellComponent* ChildComponent)
-		{
-			FVector2D RelativePosition = ChildComponent->GetPosition() - OldPosition;
-			ChildComponent->SetPosition(NewPosition + RelativePosition);
+	PositionX = NewPosition.X;
+	PositionY = NewPosition.Y;
 
-		}, bUpdateSizeRecursive);
+	const FVector2D Translation = CachedGeometry.SetPositionAbsolute(NewPosition);
+	CachedGeometry.PropagonatePositionChangesToChildren(Translation);
 
-	HandlePositionChanged();
+#if WITH_EDITOR
+	EditorPositionWithRotation = CachedGeometry.GetPositionRotatedAbsolute();
+#endif
 }
 
-void UDMXPixelMappingMatrixComponent::SetSize(const FVector2D& NewSize)
+void UDMXPixelMappingMatrixComponent::SetPositionRotated(FVector2D NewRotatedPosition)
 {
-	const FVector2D OldSize = GetSize();
-	Super::SetSize(NewSize);
+	const FVector2D Translation = CachedGeometry.SetPositionRotatedAbsolute(NewRotatedPosition);
 
-	if (OldSize != NewSize)
-	{
-		HandleSizeChanged();
-	}
-}
+	PositionX = CachedGeometry.GetPositionAbsolute().X;
+	PositionY = CachedGeometry.GetPositionAbsolute().Y;
 
-PRAGMA_DISABLE_DEPRECATION_WARNINGS
-void UDMXPixelMappingMatrixComponent::RenderWithInputAndSendDMX()
-{
-	// DEPRECATED 5.3
-	RenderAndSendDMX();
-}
-PRAGMA_ENABLE_DEPRECATION_WARNINGS
+	CachedGeometry.PropagonatePositionChangesToChildren(Translation);
 
-void UDMXPixelMappingMatrixComponent::HandlePositionChanged()
-{
 #if WITH_EDITOR
-PRAGMA_DISABLE_DEPRECATION_WARNINGS
-	if (ComponentWidget_DEPRECATED.IsValid())
-	{
-		ComponentWidget_DEPRECATED->SetPosition(GetPosition());
-	}
-PRAGMA_ENABLE_DEPRECATION_WARNINGS
+	EditorPositionWithRotation = CachedGeometry.GetPositionRotatedAbsolute();
 #endif
 }
 
-void UDMXPixelMappingMatrixComponent::HandleSizeChanged()
+void UDMXPixelMappingMatrixComponent::SetSize(const FVector2D& NewSize)
 {
-	CellSize = FVector2D::ZeroVector;
-	if (CoordinateGrid.X > 0 && CoordinateGrid.Y > 0)
-	{
-		CellSize = FVector2D(GetSize().X / CoordinateGrid.X, GetSize().Y / CoordinateGrid.Y);
-	}
+	SizeX = NewSize.X;
+	SizeY = NewSize.Y;
 
-	constexpr bool bUpdateSizeRecursive = false;
-	ForEachChildOfClass<UDMXPixelMappingMatrixCellComponent>([this](UDMXPixelMappingMatrixCellComponent* ChildComponent)
-		{
-#if WITH_EDITOR
-			if (ChildComponent->IsLockInDesigner())
-			{
-				ChildComponent->SetSize(CellSize);
-				ChildComponent->SetPosition(GetPosition() + FVector2D(CellSize * ChildComponent->GetCellCoordinate()));
-			}
-#else
-			ChildComponent->SetSize(CellSize);
-			ChildComponent->SetPosition(GetPosition() + FVector2D(CellSize * ChildComponent->GetCellCoordinate()));
-#endif // WITH_EDITOR
+	FVector2D DeltaSize;
+	FVector2D DeltaPosition;
+	CachedGeometry.SetSizeAbsolute(NewSize, DeltaSize, DeltaPosition);
+
+	// Resizing rotated components moves their pivot, and by that their position without rotation.
+	// Adjust for this so the component keeps its rotated position.
+	SetPosition(CachedGeometry.GetPositionAbsolute());
+
+	CachedGeometry.PropagonateSizeChangesToChildren(DeltaSize, DeltaPosition);
+}
+
+void UDMXPixelMappingMatrixComponent::SetRotation(double NewRotation)
+{
+	Rotation = NewRotation;
 
-		}, bUpdateSizeRecursive);
+	const double DeltaRotation = CachedGeometry.SetRotationAbsolute(NewRotation);
+	CachedGeometry.PropagonateRotationChangesToChildren(DeltaRotation);
 
 #if WITH_EDITOR
-PRAGMA_DISABLE_DEPRECATION_WARNINGS
-	if (ComponentWidget_DEPRECATED.IsValid())
-	{
-		ComponentWidget_DEPRECATED->SetSize(GetSize());
-	}
-PRAGMA_ENABLE_DEPRECATION_WARNINGS
+	EditorPositionWithRotation = CachedGeometry.GetPositionRotatedAbsolute();
 #endif
 }
 
+PRAGMA_DISABLE_DEPRECATION_WARNINGS
+void UDMXPixelMappingMatrixComponent::RenderWithInputAndSendDMX()
+{
+	// DEPRECATED 5.3
+	RenderAndSendDMX();
+}
+PRAGMA_ENABLE_DEPRECATION_WARNINGS
+
 void UDMXPixelMappingMatrixComponent::HandleMatrixChanged()
 {
 	UDMXPixelMapping* PixelMapping = GetPixelMapping();
@@ -540,7 +476,21 @@ void UDMXPixelMappingMatrixComponent::HandleMatrixChanged()
 		}
 	}
 
-	HandleSizeChanged();
+	// Set child transform
+	CellSize = FVector2D::ZeroVector;
+	if (CoordinateGrid.X > 0 && CoordinateGrid.Y > 0)
+	{
+		CellSize = FVector2D(GetSize().X / CoordinateGrid.X, GetSize().Y / CoordinateGrid.Y);
+	}
+
+	constexpr bool bUpdateSizeRecursive = false;
+	ForEachChildOfClass<UDMXPixelMappingMatrixCellComponent>([this](UDMXPixelMappingMatrixCellComponent* ChildComponent)
+		{
+			ChildComponent->SetRotation(GetRotation());
+			ChildComponent->SetSize(CellSize);
+			ChildComponent->SetPosition(GetPosition() + FVector2D(CellSize * ChildComponent->GetCellCoordinate()));
+		},
+		bUpdateSizeRecursive);
 
 	GetOnMatrixChanged().Broadcast(PixelMapping, this);
 }
diff --git a/Engine/Plugins/VirtualProduction/DMX/DMXPixelMapping/Source/DMXPixelMappingRuntime/Private/Components/DMXPixelMappingOutputComponent.cpp b/Engine/Plugins/VirtualProduction/DMX/DMXPixelMapping/Source/DMXPixelMappingRuntime/Private/Components/DMXPixelMappingOutputComponent.cpp
--- a/Engine/Plugins/VirtualProduction/DMX/DMXPixelMapping/Source/DMXPixelMappingRuntime/Private/Components/DMXPixelMappingOutputComponent.cpp
+++ b/Engine/Plugins/VirtualProduction/DMX/DMXPixelMapping/Source/DMXPixelMappingRuntime/Private/Components/DMXPixelMappingOutputComponent.cpp
@@ -2,160 +2,75 @@
 
 #include "Components/DMXPixelMappingOutputComponent.h"
 
+#include "Components/DMXPixelMappingComponentGeometryCache.h"
 #include "Components/DMXPixelMappingRendererComponent.h"
-
-#if WITH_EDITOR
-#include "DMXPixelMappingComponentWidget.h"
-#endif // WITH_EDITOR
+#include "DMXPixelMapping.h"
+#include "DMXPixelMappingMainStreamObjectVersion.h"
 
 
 #define LOCTEXT_NAMESPACE "DMXPixelMappingOutputComponent"
 
-
 UDMXPixelMappingOutputComponent::UDMXPixelMappingOutputComponent()
-	: PositionX(0.f)
-	, PositionY(0.f)
-	, SizeX(1.f)
-	, SizeY(1.f)
 {
 #if WITH_EDITOR
 	bLockInDesigner = false;
 	bVisibleInDesigner = true;
 	ZOrder = 0;
 #endif // WITH_EDITOR
+
+	// It is important to initialize with properties, as the getters make use of the geometry cache.
+	CachedGeometry.Initialize(this, FVector2D(PositionX, PositionY), FVector2D(SizeX, SizeY), Rotation);
 }
 
-#if WITH_EDITOR
-void UDMXPixelMappingOutputComponent::PostEditChangeProperty(FPropertyChangedEvent& PropertyChangedEvent)
+void UDMXPixelMappingOutputComponent::PostInitProperties()
 {
-	// Call the parent at the first place
-	Super::PostEditChangeProperty(PropertyChangedEvent);
-	
-PRAGMA_DISABLE_DEPRECATION_WARNINGS
-	// Handle deprecated component widget changes
-	if (PropertyChangedEvent.GetPropertyName() == GET_MEMBER_NAME_CHECKED(UDMXPixelMappingOutputComponent, bVisibleInDesigner))
-	{
-		const EVisibility NewVisiblity = IsVisible() ? EVisibility::HitTestInvisible : EVisibility::Hidden;
-
-		if (ComponentWidget_DEPRECATED.IsValid())
-		{
-			ComponentWidget_DEPRECATED->SetVisibility(NewVisiblity);
-		}
+	Super::PostInitProperties();
 
-		constexpr bool bSetVisibilityRecursive = true;
-		ForEachChildOfClass<UDMXPixelMappingOutputComponent>([NewVisiblity](UDMXPixelMappingOutputComponent* ChildComponent)
-			{
-				if (TSharedPtr<FDMXPixelMappingComponentWidget> ChildComponentWidget = ChildComponent->GetComponentWidget())
-				{
-					ChildComponentWidget->SetVisibility(NewVisiblity);
-				}
-			}, bSetVisibilityRecursive);
-	}
-PRAGMA_ENABLE_DEPRECATION_WARNINGS
+#if WITH_EDITOR
+	EditorPositionWithRotation = GetPosition();
+#endif
 }
-#endif // WITH_EDITOR
 
-#if WITH_EDITOR
-void UDMXPixelMappingOutputComponent::PreEditUndo()
+void UDMXPixelMappingOutputComponent::PostLoad()
 {
-	Super::PreEditUndo();
+	Super::PostLoad();
 
-	PreEditUndoChildren = Children;
+	// It is important to initialize with properties, as the getters make use of the geometry cache.
+	CachedGeometry.Initialize(this, FVector2D(PositionX, PositionY), FVector2D(SizeX, SizeY), Rotation);
+
+#if WITH_EDITOR
+	EditorPositionWithRotation = CachedGeometry.GetPositionRotatedAbsolute();
+#endif
 }
-#endif // WITH_EDITOR
 
 #if WITH_EDITOR
-void UDMXPixelMappingOutputComponent::PostEditUndo()
+void UDMXPixelMappingOutputComponent::PostEditChangeProperty(FPropertyChangedEvent& PropertyChangedEvent)
 {
-	Super::PostEditUndo();
-
-	EVisibility NewVisibility = IsVisible() ? EVisibility::HitTestInvisible : EVisibility::Hidden;
-
-	UpdateComponentWidget(NewVisibility);
-
-	if (Children.Num() < PreEditUndoChildren.Num())
-	{	
-		// Undo Add Children (remove them again)
-		for (UDMXPixelMappingBaseComponent* PreEditUndoChild : PreEditUndoChildren)
-		{
-			if (UDMXPixelMappingOutputComponent* PreEditUndoOutputComponent = Cast<UDMXPixelMappingOutputComponent>(PreEditUndoChild))
-			{
-				const EVisibility NewChildVisibility = PreEditUndoOutputComponent->IsVisible() ? EVisibility::HitTestInvisible : EVisibility::Hidden;
+	Super::PostEditChangeProperty(PropertyChangedEvent);
 
-				constexpr bool bWithChildrenRecursive = false;
-				PreEditUndoOutputComponent->UpdateComponentWidget(NewChildVisibility, bWithChildrenRecursive);
-			}	
-		}
-	}
-	else if (Children.Num() > PreEditUndoChildren.Num())
+	const FName PropertyName = PropertyChangedEvent.GetPropertyName();
+	const FName MemberPropertyName = PropertyChangedEvent.GetMemberPropertyName();
+	if (MemberPropertyName == GET_MEMBER_NAME_CHECKED(UDMXPixelMappingOutputComponent, EditorPositionWithRotation))
 	{
-		// Undo Remove Children (add them again)
-		for (UDMXPixelMappingBaseComponent* PreEditUndoChild : PreEditUndoChildren)
-		{
-			if (UDMXPixelMappingOutputComponent* PreEditUndoOutputComponent = Cast<UDMXPixelMappingOutputComponent>(PreEditUndoChild))
-			{
-				const EVisibility NewChildVisibility = IsVisible() ? EVisibility::HitTestInvisible : EVisibility::Hidden;
-
-				constexpr bool bWithChildrenRecursive = false;
-				PreEditUndoOutputComponent->UpdateComponentWidget(NewChildVisibility, bWithChildrenRecursive);
-			}
-		}
+		SetPositionRotated(EditorPositionWithRotation);
+	} 
+	else if (PropertyName == GET_MEMBER_NAME_CHECKED(UDMXPixelMappingOutputComponent, SizeX) ||
+		PropertyName == GET_MEMBER_NAME_CHECKED(UDMXPixelMappingOutputComponent, SizeY))
+	{
+		SetSize(FVector2D(SizeX, SizeY));
 	}
-
-	PreEditUndoChildren.Reset();
-}
-#endif // WITH_EDITOR
-
-void UDMXPixelMappingOutputComponent::BeginDestroy()
-{
-	Super::BeginDestroy();
-
-#if WITH_EDITOR
-	PRAGMA_DISABLE_DEPRECATION_WARNINGS
-	if (ComponentWidget_DEPRECATED.IsValid())
+	else if(PropertyName == GET_MEMBER_NAME_CHECKED(UDMXPixelMappingOutputComponent, Rotation))
 	{
-		ComponentWidget_DEPRECATED->RemoveFromCanvas();
-
-		// Should have released all references by now
-		ensureMsgf(ComponentWidget_DEPRECATED.GetSharedReferenceCount() == 1, TEXT("Detected Reference to Component Widget the component is destroyed."));
+		SetRotation(Rotation);
 	}
-	PRAGMA_ENABLE_DEPRECATION_WARNINGS
-#endif
 }
+#endif // WITH_EDITOR
 
 bool UDMXPixelMappingOutputComponent::CanBeMovedTo(const UDMXPixelMappingBaseComponent* Component) const
 {
 	return false;
 }
 
-void UDMXPixelMappingOutputComponent::NotifyAddedToParent()
-{
-	Super::NotifyAddedToParent();
-
-#if WITH_EDITOR
-	UpdateComponentWidget(EVisibility::HitTestInvisible);
-#endif // WITH_EDITOR
-}
-
-void UDMXPixelMappingOutputComponent::NotifyRemovedFromParent()
-{
-	Super::NotifyRemovedFromParent();
-
-#if WITH_EDITOR
-	PRAGMA_DISABLE_DEPRECATION_WARNINGS
-	if (ComponentWidget_DEPRECATED.IsValid())
-	{
-		ComponentWidget_DEPRECATED->RemoveFromCanvas();
-
-		// Should have released all references by now
-		ensureMsgf(ComponentWidget_DEPRECATED.GetSharedReferenceCount() == 1, TEXT("Detected Reference to Component Widget the component is destroyed."));
-
-		ComponentWidget_DEPRECATED.Reset();
-	}
-	PRAGMA_ENABLE_DEPRECATION_WARNINGS
-#endif // WITH_EDITOR
-}
-
 #if WITH_EDITOR
 const FText UDMXPixelMappingOutputComponent::GetPaletteCategory()
 {
@@ -165,22 +80,6 @@ const FText UDMXPixelMappingOutputComponent::GetPaletteCategory()
 }
 #endif // WITH_EDITOR
 
-#if WITH_EDITOR
-PRAGMA_DISABLE_DEPRECATION_WARNINGS
-TSharedRef<FDMXPixelMappingComponentWidget> UDMXPixelMappingOutputComponent::BuildSlot(TSharedRef<SConstraintCanvas> InCanvas)
-{
-	ComponentWidget_DEPRECATED = MakeShared<FDMXPixelMappingComponentWidget>();
-	ComponentWidget_DEPRECATED->AddToCanvas(InCanvas, ZOrder);
-
-	EVisibility NewVisibility = IsVisible() ? EVisibility::HitTestInvisible : EVisibility::Hidden;
-
-	UpdateComponentWidget(NewVisibility);
-
-	return ComponentWidget_DEPRECATED.ToSharedRef();
-}
-PRAGMA_ENABLE_DEPRECATION_WARNINGS
-#endif // WITH_EDITOR
-
 #if WITH_EDITOR
 bool UDMXPixelMappingOutputComponent::IsVisible() const
 {
@@ -196,34 +95,6 @@ void UDMXPixelMappingOutputComponent::SetZOrder(int32 NewZOrder)
 	const int32 DeltaZOrder = NewZOrder - ZOrder;
 	
 	ZOrder = NewZOrder;
-
-
-PRAGMA_DISABLE_DEPRECATION_WARNINGS
-	// Handle deprecated component widget changes
-	if (ComponentWidget_DEPRECATED.IsValid())
-	{
-		ComponentWidget_DEPRECATED->SetZOrder(ZOrder);
-	}
-
-	constexpr bool bRecursive = true;
-	ForEachChild([DeltaZOrder, this](UDMXPixelMappingBaseComponent* ChildComponent)
-		{
-			if (UDMXPixelMappingOutputComponent* ChildOutputComponent = Cast<UDMXPixelMappingOutputComponent>(ChildComponent))
-			{
-				const int32 NewChildZOrder = ChildOutputComponent->GetZOrder() + DeltaZOrder;
-				if (TSharedPtr<FDMXPixelMappingComponentWidget> ChildComponentWidget_DEPRECATED = ChildOutputComponent->GetComponentWidget())
-				{
-					ChildOutputComponent->SetZOrder(NewChildZOrder);
-
-					// Apply to the UI
-					if (ChildOutputComponent->ComponentWidget_DEPRECATED.IsValid())
-					{
-						ChildOutputComponent->ComponentWidget_DEPRECATED->SetZOrder(ZOrder);
-					}
-				}
-			}
-		}, bRecursive);
-PRAGMA_ENABLE_DEPRECATION_WARNINGS
 }
 #endif // WITH_EDITOR
 
@@ -236,44 +107,110 @@ bool UDMXPixelMappingOutputComponent::IsOverParent() const
 
 bool UDMXPixelMappingOutputComponent::IsOverPosition(const FVector2D& OtherPosition) const
 {
-	return
-		PositionX <= OtherPosition.X &&
-		PositionY <= OtherPosition.Y &&
-		PositionX + SizeX >= OtherPosition.X &&
-		PositionY + SizeY >= OtherPosition.Y;
+	FVector2D A;
+	FVector2D B;
+	FVector2D C;
+	FVector2D D;
+	CachedGeometry.GetEdges(A, B, C, D);
+
+	const FVector2D AM = OtherPosition - A;
+	const FVector2D AB = B - A;
+	const FVector2D AD = D - A;
+	
+	const double AMDotAB = FVector2D::DotProduct(AM, AB);
+	const double MagnituteSquaredAB = FVector2D::DotProduct(AB, AB);
+	const double AMDotAD = FVector2D::DotProduct(AM, AD);
+	const double MagnituteSquaredAD = FVector2D::DotProduct(AD, AD);
+
+	constexpr double Tolerance = 0.1;
+	const bool bPointsTowardsRectangleArea = -Tolerance <= AMDotAB && -Tolerance <= AMDotAD;
+	const bool bIsAMSmallerThanSides = AMDotAB <= MagnituteSquaredAB + Tolerance && AMDotAD < MagnituteSquaredAD + Tolerance;
+
+	return bPointsTowardsRectangleArea && bIsAMSmallerThanSides;
 }
 
 bool UDMXPixelMappingOutputComponent::OverlapsComponent(UDMXPixelMappingOutputComponent* Other) const
 {
+	// DEPRECATED 5.4
 	if (Other)
 	{
-		FVector2D ThisPosition = GetPosition();
-		FVector2D OtherPosition = Other->GetPosition();
-
-		FBox2D ThisBox = FBox2D(ThisPosition, ThisPosition + GetSize());
-		FBox2D OtherBox = FBox2D(OtherPosition, OtherPosition + Other->GetSize());
-
-		return ThisBox.Intersect(OtherBox);
+		FVector2D A;
+		FVector2D B;
+		FVector2D C;
+		FVector2D D;
+		CachedGeometry.GetEdges(A, B, C, D);
+
+		return
+			Other->IsOverPosition(A) ||
+			Other->IsOverPosition(B) ||
+			Other->IsOverPosition(C) ||
+			Other->IsOverPosition(D);
 	}
 
 	return false;
 }
 
-void UDMXPixelMappingOutputComponent::SetPosition(const FVector2D& Position) 
+FVector2D UDMXPixelMappingOutputComponent::GetPosition() const
 {
-	PositionX = Position.X;
-	PositionY = Position.Y;
+	return CachedGeometry.GetPositionAbsolute();
 }
 
-FVector2D UDMXPixelMappingOutputComponent::GetPosition() const
+void UDMXPixelMappingOutputComponent::SetPosition(const FVector2D& NewPosition) 
 {
-	return FVector2D(PositionX, PositionY);
+	PositionX = NewPosition.X;
+	PositionY = NewPosition.Y;
+
+	CachedGeometry.SetPositionAbsolute(NewPosition);
+
+#if WITH_EDITOR
+	EditorPositionWithRotation = CachedGeometry.GetPositionRotatedAbsolute();
+#endif
 }
 
-void UDMXPixelMappingOutputComponent::SetSize(const FVector2D& Size) 
+FVector2D UDMXPixelMappingOutputComponent::GetPositionRotated() const
 {
-	SizeX = Size.X;
-	SizeY = Size.Y;
+	return CachedGeometry.GetPositionRotatedAbsolute();
+}
+
+void UDMXPixelMappingOutputComponent::SetPositionRotated(FVector2D NewRotatedPosition)
+{
+	CachedGeometry.SetPositionRotatedAbsolute(NewRotatedPosition);
+
+	PositionX = CachedGeometry.GetPositionAbsolute().X;
+	PositionY = CachedGeometry.GetPositionAbsolute().Y;
+
+#if WITH_EDITOR
+	EditorPositionWithRotation = CachedGeometry.GetPositionRotatedAbsolute();
+#endif
+}
+
+void UDMXPixelMappingOutputComponent::SetSize(const FVector2D& NewSize)
+{
+	SizeX = NewSize.X;
+	SizeY = NewSize.Y;
+
+	CachedGeometry.SetSizeAbsolute(NewSize);
+}
+
+FVector2D UDMXPixelMappingOutputComponent::GetSize() const
+{
+	return CachedGeometry.GetSizeAbsolute();
+}
+
+void UDMXPixelMappingOutputComponent::SetRotation(double NewRotation)
+{
+	Rotation = NewRotation;
+
+	CachedGeometry.SetRotationAbsolute(NewRotation);
+
+#if WITH_EDITOR
+	EditorPositionWithRotation = CachedGeometry.GetPositionRotatedAbsolute();
+#endif
+}
+
+double UDMXPixelMappingOutputComponent::GetRotation() const
+{
+	return CachedGeometry.GetRotationAbsolute();
 }
 
 void UDMXPixelMappingOutputComponent::InvalidatePixelMapRenderer()
@@ -344,32 +281,4 @@ void UDMXPixelMappingOutputComponent::ZOrderTopmost()
 }
 #endif // WITH_EDITOR
 
-#if WITH_EDITOR
-void UDMXPixelMappingOutputComponent::UpdateComponentWidget(EVisibility NewVisibility, bool bWithChildrenRecursive)
-{
-	PRAGMA_DISABLE_DEPRECATION_WARNINGS
-	if (ComponentWidget_DEPRECATED.IsValid())
-	{
-		ComponentWidget_DEPRECATED->SetPosition(FVector2D(PositionX, PositionY));
-		ComponentWidget_DEPRECATED->SetSize(FVector2D(SizeX, SizeY));
-		ComponentWidget_DEPRECATED->SetVisibility(NewVisibility);
-		ComponentWidget_DEPRECATED->SetColor(GetEditorColor());
-		ComponentWidget_DEPRECATED->SetLabelText(FText::FromString(GetUserFriendlyName()));
-	}
-
-	if (bWithChildrenRecursive)
-	{
-		for (UDMXPixelMappingBaseComponent* ChildComponent : Children)
-		{
-			if (UDMXPixelMappingOutputComponent* ChildOutputComponent = Cast<UDMXPixelMappingOutputComponent>(ChildComponent))
-			{
-				// Recursive for all
-				ChildOutputComponent->UpdateComponentWidget(NewVisibility, bWithChildrenRecursive);
-			}
-		}
-	}
-	PRAGMA_ENABLE_DEPRECATION_WARNINGS
-}
-#endif // WITH_EDITOR
-
 #undef LOCTEXT_NAMESPACE
diff --git a/Engine/Plugins/VirtualProduction/DMX/DMXPixelMapping/Source/DMXPixelMappingRuntime/Private/Components/DMXPixelMappingRendererComponent.cpp b/Engine/Plugins/VirtualProduction/DMX/DMXPixelMapping/Source/DMXPixelMappingRuntime/Private/Components/DMXPixelMappingRendererComponent.cpp
--- a/Engine/Plugins/VirtualProduction/DMX/DMXPixelMapping/Source/DMXPixelMappingRuntime/Private/Components/DMXPixelMappingRendererComponent.cpp
+++ b/Engine/Plugins/VirtualProduction/DMX/DMXPixelMapping/Source/DMXPixelMappingRuntime/Private/Components/DMXPixelMappingRendererComponent.cpp
@@ -4,6 +4,7 @@
 
 #include "Async/Async.h"
 #include "Blueprint/UserWidget.h"
+#include "Components/DMXPixelMappingComponentGeometryCache.h"
 #include "Components/DMXPixelMappingFixtureGroupComponent.h"
 #include "Components/DMXPixelMappingFixtureGroupItemComponent.h"
 #include "Components/DMXPixelMappingMatrixComponent.h"
@@ -672,31 +673,6 @@ void UDMXPixelMappingRendererComponent::EmptyDownsampleBuffer()
 	DownsampleBuffer_DEPRECATED.Empty();
 }
 
-#if WITH_EDITOR
-TSharedRef<SWidget> UDMXPixelMappingRendererComponent::TakeWidget()
-{	
-	// DEPRECATED 5.3
-
-	if (!ComponentsCanvas_DEPRECATED.IsValid())
-	{
-		ComponentsCanvas_DEPRECATED =
-			SNew(SConstraintCanvas);
-	}
-
-	PRAGMA_DISABLE_DEPRECATION_WARNINGS
-	ForEachChild([&](UDMXPixelMappingBaseComponent* InComponent) {
-			if (UDMXPixelMappingOutputComponent* Component = Cast<UDMXPixelMappingOutputComponent>(InComponent))
-			{
-				// Build all child DMX pixel mapping slots
-				Component->BuildSlot(ComponentsCanvas_DEPRECATED.ToSharedRef());
-			}
-		}, true);
-	PRAGMA_ENABLE_DEPRECATION_WARNINGS
-
-	return ComponentsCanvas_DEPRECATED.ToSharedRef();
-}
-#endif // WITH_EDITOR
-
 int32 UDMXPixelMappingRendererComponent::GetTotalDownsamplePixelCount()
 {
 	// DEPRECATED 5.3
diff --git a/Engine/Plugins/VirtualProduction/DMX/DMXPixelMapping/Source/DMXPixelMappingRuntime/Private/Components/DMXPixelMappingScreenComponent.cpp b/Engine/Plugins/VirtualProduction/DMX/DMXPixelMapping/Source/DMXPixelMappingRuntime/Private/Components/DMXPixelMappingScreenComponent.cpp
--- a/Engine/Plugins/VirtualProduction/DMX/DMXPixelMapping/Source/DMXPixelMappingRuntime/Private/Components/DMXPixelMappingScreenComponent.cpp
+++ b/Engine/Plugins/VirtualProduction/DMX/DMXPixelMapping/Source/DMXPixelMappingRuntime/Private/Components/DMXPixelMappingScreenComponent.cpp
@@ -94,25 +94,6 @@ void UDMXPixelMappingScreenComponent::PostEditChangeChainProperty(FPropertyChang
 			ScreenComponentBox->RebuildGrid(GridParams);
 		}
 	}
-	if (PropertyChangedChainEvent.ChangeType != EPropertyChangeType::Interactive)
-	{
-		if (PropertyChangedChainEvent.GetPropertyName() == UDMXPixelMappingOutputComponent::GetPositionXPropertyName() ||
-			PropertyChangedChainEvent.GetPropertyName() == UDMXPixelMappingOutputComponent::GetPositionYPropertyName())
-		{
-			if (ComponentWidget_DEPRECATED.IsValid())
-			{
-				ComponentWidget_DEPRECATED->SetPosition(GetPosition());
-			}
-		}
-		else if (PropertyChangedChainEvent.GetPropertyName() == UDMXPixelMappingOutputComponent::GetSizeXPropertyName() ||
-			PropertyChangedChainEvent.GetPropertyName() == UDMXPixelMappingOutputComponent::GetSizeYPropertyName())
-		{
-			if (ComponentWidget_DEPRECATED.IsValid())
-			{
-				ComponentWidget_DEPRECATED->SetSize(GetPosition());
-			}
-		}
-	}
 	PRAGMA_ENABLE_DEPRECATION_WARNINGS
 }
 #endif // WITH_EDITOR
@@ -124,37 +105,6 @@ const FText UDMXPixelMappingScreenComponent::GetPaletteCategory()
 }
 #endif // WITH_EDITOR
 
-#if WITH_EDITOR
-PRAGMA_DISABLE_DEPRECATION_WARNINGS
-TSharedRef<FDMXPixelMappingComponentWidget> UDMXPixelMappingScreenComponent::BuildSlot(TSharedRef<SConstraintCanvas> InCanvas)
-{
-	if (!ComponentWidget_DEPRECATED.IsValid())
-	{
-		ScreenComponentBox = 
-			SNew(SDMXPixelMappingScreenComponentBox)
-			.NumXCells(NumXCells)
-			.NumYCells(NumYCells)
-			.Distribution(Distribution)
-			.PixelFormat(PixelFormat)
-			.LocalUniverse(LocalUniverse)
-			.StartAddress(StartAddress)
-			.bShowAddresses(bShowAddresses)
-			.bShowUniverse(bShowUniverse);
-
-		ComponentWidget_DEPRECATED = MakeShared<FDMXPixelMappingComponentWidget>(ScreenComponentBox, nullptr);
-
-		ComponentWidget_DEPRECATED->AddToCanvas(InCanvas, ZOrder);
-		ComponentWidget_DEPRECATED->SetPosition(GetPosition());
-		ComponentWidget_DEPRECATED->SetSize(GetSize());
-		ComponentWidget_DEPRECATED->SetColor(GetEditorColor());
-		ComponentWidget_DEPRECATED->SetLabelText(FText::FromString(GetUserFriendlyName()));
-	}
-	
-	return ComponentWidget_DEPRECATED.ToSharedRef();
-}
-PRAGMA_ENABLE_DEPRECATION_WARNINGS
-#endif // WITH_EDITOR
-
 const FName& UDMXPixelMappingScreenComponent::GetNamePrefix()
 {
 	static FName NamePrefix = TEXT("DEPRECATED DMX Screen");
@@ -463,32 +413,4 @@ void UDMXPixelMappingScreenComponent::ForEachPixel(ForEachPixelCallback InCallba
 	}
 }
 
-void UDMXPixelMappingScreenComponent::SetPosition(const FVector2D& NewPosition)
-{
-	Super::SetPosition(NewPosition);
-
-#if WITH_EDITOR
-	PRAGMA_DISABLE_DEPRECATION_WARNINGS
-	if (ComponentWidget_DEPRECATED.IsValid())
-	{
-		ComponentWidget_DEPRECATED->SetPosition(GetPosition());
-	}
-	PRAGMA_ENABLE_DEPRECATION_WARNINGS
-#endif
-}
-
-void UDMXPixelMappingScreenComponent::SetSize(const FVector2D& NewSize)
-{
-	Super::SetSize(NewSize);
-
-#if WITH_EDITOR
-	PRAGMA_DISABLE_DEPRECATION_WARNINGS
-	if (ComponentWidget_DEPRECATED.IsValid())
-	{
-		ComponentWidget_DEPRECATED->SetSize(GetSize());
-	}
-	PRAGMA_ENABLE_DEPRECATION_WARNINGS
-#endif
-}
-
 #undef LOCTEXT_NAMESPACE
diff --git a/Engine/Plugins/VirtualProduction/DMX/DMXPixelMapping/Source/DMXPixelMappingRuntime/Private/DMXPixelMapping.cpp b/Engine/Plugins/VirtualProduction/DMX/DMXPixelMapping/Source/DMXPixelMappingRuntime/Private/DMXPixelMapping.cpp
--- a/Engine/Plugins/VirtualProduction/DMX/DMXPixelMapping/Source/DMXPixelMappingRuntime/Private/DMXPixelMapping.cpp
+++ b/Engine/Plugins/VirtualProduction/DMX/DMXPixelMapping/Source/DMXPixelMappingRuntime/Private/DMXPixelMapping.cpp
@@ -131,30 +131,6 @@ UDMXPixelMappingBaseComponent* UDMXPixelMapping::FindComponent(const FName& InNa
 	return FoundComponent;
 }
 
-#if WITH_EDITOR
-UDMXPixelMappingOutputComponent* UDMXPixelMapping::FindComponent(TSharedPtr<SWidget> InWidget) const
-{
-	PRAGMA_DISABLE_DEPRECATION_WARNINGS
-	UDMXPixelMappingOutputComponent* FoundComponent = nullptr;
-
-	ForEachComponentOfClass<UDMXPixelMappingOutputComponent>([&](UDMXPixelMappingOutputComponent* InComponent)
-		{
-			if (TSharedPtr<FDMXPixelMappingComponentWidget> ComponentWidget = InComponent->GetComponentWidget())
-			{
-				if(ComponentWidget->GetComponentBox() == InWidget)
-				{
-					FoundComponent = InComponent;
-				}
-			}
-		});
-
-	return FoundComponent;
-	PRAGMA_ENABLE_DEPRECATION_WARNINGS
-}
-
-#endif // WITH_EDITOR
-
-
 void UDMXPixelMapping::RemoveComponent(UDMXPixelMappingBaseComponent* InComponent)
 {
 	ensureMsgf(InComponent, TEXT("Trying to remove invalid component."));
diff --git a/Engine/Plugins/VirtualProduction/DMX/DMXPixelMapping/Source/DMXPixelMappingRuntime/Public/Components/DMXPixelMappingFixtureGroupComponent.h b/Engine/Plugins/VirtualProduction/DMX/DMXPixelMapping/Source/DMXPixelMappingRuntime/Public/Components/DMXPixelMappingFixtureGroupComponent.h
--- a/Engine/Plugins/VirtualProduction/DMX/DMXPixelMapping/Source/DMXPixelMappingRuntime/Public/Components/DMXPixelMappingFixtureGroupComponent.h
+++ b/Engine/Plugins/VirtualProduction/DMX/DMXPixelMapping/Source/DMXPixelMappingRuntime/Public/Components/DMXPixelMappingFixtureGroupComponent.h
@@ -26,12 +26,10 @@ public:
 	UDMXPixelMappingFixtureGroupComponent();
 
 	//~ Begin UObject implementation
-	virtual void PostInitProperties() override;
-	virtual void PostLoad() override;
 #if WITH_EDITOR
 	virtual void PostEditChangeProperty(FPropertyChangedEvent& PropertyChangedEvent) override;
 	virtual void PostEditUndo() override;
-#endif // WITH_EDITOR
+#endif
 	//~ End UObject implementation
 
 	//~ Begin UDMXPixelMappingBaseComponent implementation
@@ -44,12 +42,14 @@ public:
 	//~ End UDMXPixelMappingBaseComponent implementation
 
 	//~ Begin UDMXPixelMappingOutputComponent implementation
+	virtual void SetPosition(const FVector2D& NewPosition) override;
+	virtual void SetPositionRotated(FVector2D NewRotatedPosition) override;
+	virtual void SetSize(const FVector2D& NewSize) override;
+	virtual void SetRotation(double NewRotation) override;
 #if WITH_EDITOR
 	virtual bool IsExposedToTemplate() { return true; }
 	virtual const FText GetPaletteCategory() override;
-#endif // WITH_EDITOR
-	virtual void SetPosition(const FVector2D& NewPosition) override;
-	virtual void SetSize(const FVector2D& NewSize) override;
+#endif
 	//~ End UDMXPixelMappingOutputComponent implementation
 
 	/** Returns a delegate broadcasted when the DMX Library changed */
@@ -66,12 +66,6 @@ public:
 	TObjectPtr<UDMXPixelMappingLayoutScript> LayoutScript;
 
 private:
-	/** Handles changes in position */
-	void HandlePositionChanged();
-
-	/** Holds the last set size */
-	FVector2D LastPosition;
-
 	/** Returns a delegate broadcasted when the DMX Library changed */
 	FSimpleMulticastDelegate OnDMXLibraryChangedDelegate;
 
diff --git a/Engine/Plugins/VirtualProduction/DMX/DMXPixelMapping/Source/DMXPixelMappingRuntime/Public/Components/DMXPixelMappingMatrixCellComponent.h b/Engine/Plugins/VirtualProduction/DMX/DMXPixelMapping/Source/DMXPixelMappingRuntime/Public/Components/DMXPixelMappingMatrixCellComponent.h
--- a/Engine/Plugins/VirtualProduction/DMX/DMXPixelMapping/Source/DMXPixelMappingRuntime/Public/Components/DMXPixelMappingMatrixCellComponent.h
+++ b/Engine/Plugins/VirtualProduction/DMX/DMXPixelMapping/Source/DMXPixelMappingRuntime/Public/Components/DMXPixelMappingMatrixCellComponent.h
@@ -51,17 +51,11 @@ public:
 
 	//~ Begin UDMXPixelMappingOutputComponent implementation
 #if WITH_EDITOR
-	PRAGMA_DISABLE_DEPRECATION_WARNINGS // FDMXPixelMappingComponentWidget is deprecated
-	virtual TSharedRef<FDMXPixelMappingComponentWidget> BuildSlot(TSharedRef<SConstraintCanvas> InCanvas) override;
-	PRAGMA_ENABLE_DEPRECATION_WARNINGS
-	virtual bool IsVisible() const override;
 	virtual FLinearColor GetEditorColor() const override;
 #endif // WITH_EDITOR	
 	virtual bool IsOverParent() const override;
 	virtual int32 GetDownsamplePixelIndex() const override { return DownsamplePixelIndex; }
 	virtual void QueueDownsample() override;
-	virtual void SetPosition(const FVector2D& NewPosition) override;
-	virtual void SetSize(const FVector2D& NewSize) override;
 	//~ End UDMXPixelMappingOutputComponent implementation
 
 	void SetCellCoordinate(FIntPoint InCellCoordinate);
diff --git a/Engine/Plugins/VirtualProduction/DMX/DMXPixelMapping/Source/DMXPixelMappingRuntime/Public/Components/DMXPixelMappingMatrixComponent.h b/Engine/Plugins/VirtualProduction/DMX/DMXPixelMapping/Source/DMXPixelMappingRuntime/Public/Components/DMXPixelMappingMatrixComponent.h
--- a/Engine/Plugins/VirtualProduction/DMX/DMXPixelMapping/Source/DMXPixelMappingRuntime/Public/Components/DMXPixelMappingMatrixComponent.h
+++ b/Engine/Plugins/VirtualProduction/DMX/DMXPixelMapping/Source/DMXPixelMappingRuntime/Public/Components/DMXPixelMappingMatrixComponent.h
@@ -46,7 +46,6 @@ public:
 	virtual void PostLoad() override;
 	virtual void PostInitProperties() override;
 #if WITH_EDITOR
-	virtual void PreEditChange(FProperty* PropertyAboutToChange) override;
 	virtual void PostEditChangeProperty(FPropertyChangedEvent& PropertyChangedEvent) override;
 	virtual void PostEditChangeChainProperty(FPropertyChangedChainEvent& PropertyChangedChainEvent) override;
 #endif // WITH_EDITOR
@@ -65,12 +64,14 @@ public:
 	// ~End UDMXPixelMappingBaseComponent interface
 
 	// ~Begin UDMXPixelMappingOutputComponent interface
+	virtual void SetPosition(const FVector2D& NewPosition) override;
+	virtual void SetPositionRotated(FVector2D NewRotatedPosition) override;
+	virtual void SetSize(const FVector2D& NewSize) override;
+	virtual void SetRotation(double NewRotation) override;
 #if WITH_EDITOR
 	virtual const FText GetPaletteCategory() override;
 #endif // WITH_EDITOR
 	virtual bool IsOverParent() const override;
-	virtual void SetPosition(const FVector2D& NewPosition) override;
-	virtual void SetSize(const FVector2D& NewSize) override;
 	// ~End UDMXPixelMappingOutputComponent interface
 	
 	//~ Begin UDMXPixelMappingOutputDMXComponent implementation
@@ -80,12 +81,6 @@ public:
 	PRAGMA_ENABLE_DEPRECATION_WARNINGS
 	//~ End UDMXPixelMappingOutputDMXComponent implementation
 
-	/** Handles changes in position */
-	void HandlePositionChanged();
-
-	/** Handles changes in size or in matrix */
-	void HandleSizeChanged();
-
 	/** Handles changes in size or in matrix */
 	void HandleMatrixChanged();
 
@@ -104,10 +99,6 @@ protected:
 	TArray<UDMXPixelMappingBaseComponent*> PreEditUndoMatrixCellChildren;
 #endif // WITH_EDITORONLY_DATA
 
-private:
-	/** Position before it was changed */
-	FVector2D PreEditChangePosition;
-
 public:
 #if WITH_EDITORONLY_DATA
 	UPROPERTY()
diff --git a/Engine/Plugins/VirtualProduction/DMX/DMXPixelMapping/Source/DMXPixelMappingRuntime/Public/Components/DMXPixelMappingOutputComponent.h b/Engine/Plugins/VirtualProduction/DMX/DMXPixelMapping/Source/DMXPixelMappingRuntime/Public/Components/DMXPixelMappingOutputComponent.h
--- a/Engine/Plugins/VirtualProduction/DMX/DMXPixelMapping/Source/DMXPixelMappingRuntime/Public/Components/DMXPixelMappingOutputComponent.h
+++ b/Engine/Plugins/VirtualProduction/DMX/DMXPixelMapping/Source/DMXPixelMappingRuntime/Public/Components/DMXPixelMappingOutputComponent.h
@@ -3,25 +3,18 @@
 #pragma once
 
 #include "Components/DMXPixelMappingBaseComponent.h"
-
-#if WITH_EDITOR
-#include "Widgets/Layout/SConstraintCanvas.h"
-#endif
+#include "Components/DMXPixelMappingComponentGeometryCache.h"
 
 #include "DMXPixelMappingOutputComponent.generated.h"
 
 struct FDMXPixelMappingLayoutToken;
 class SBox;
 class UDMXEntityFixturePatch;
+class UDMXPixelMappingComponentGeometryCache;
 class UDMXPixelMappingRendererComponent;
 namespace UE::DMXPixelMapping::Rendering::PixelMapRenderer { class FPixelMapRenderElement; }
 
 
-#if WITH_EDITOR
-enum class EDMXPixelMappingComponentLabelAlignment : uint8;
-class FDMXPixelMappingComponentWidget;
-#endif
-
 
 /**
  * Base class for all Designer and configurable components
@@ -36,23 +29,19 @@ public:
 	/** Default Constructor */
 	UDMXPixelMappingOutputComponent();
 
-protected:
 	//~ Begin UObject Interface
+	virtual void PostInitProperties() override;
+	virtual void PostLoad() override;
 #if WITH_EDITOR
 	virtual void PostEditChangeProperty(FPropertyChangedEvent& PropertyChangedEvent) override;
-	virtual void PreEditUndo() override;
-	virtual void PostEditUndo() override;
 #endif // WITH_EDITOR
-	virtual void BeginDestroy() override;
 	//~ End UObject Interface
 
 public:
 	//~ Begin DMXPixelMappingBaseComponent interface
 	virtual bool CanBeMovedTo(const UDMXPixelMappingBaseComponent* Component) const override;
-	virtual void NotifyAddedToParent() override;
-	virtual void NotifyRemovedFromParent() override;
 	//~ End DMXPixelMappingBaseComponent interface
- 
+
 	/*----------------------------------
 		UDMXPixelMappingOutputComponent Interface
 	----------------------------------*/
@@ -63,12 +52,6 @@ public:
 	/** Returns the text of palette category*/
 	virtual const FText GetPaletteCategory();
 
-	/** Rebuild widget for designer view */
-	UE_DEPRECATED(5.1, "Pixel Mapping Components no longer hold their own widget, in an effort to separate Views from Data.")
-	PRAGMA_DISABLE_DEPRECATION_WARNINGS // FDMXPixelMappingComponentWidget is deprecated
-	virtual TSharedRef<FDMXPixelMappingComponentWidget> BuildSlot(TSharedRef<SConstraintCanvas> InCanvas);
-	PRAGMA_ENABLE_DEPRECATION_WARNINGS
-
 	/** Whether component should be visible */
 	virtual bool IsVisible() const;
 
@@ -92,27 +75,40 @@ public:
 	virtual bool IsOverPosition(const FVector2D& Position) const;
 
 	/** Returns true if the component overlaps the other */
-	virtual bool OverlapsComponent(UDMXPixelMappingOutputComponent* Other) const;
+	UE_DEPRECATED(5.4, "Removed without replacement. GetEdges and test against other if this method is needed.")
+		virtual bool OverlapsComponent(UDMXPixelMappingOutputComponent* Other) const;
 
 	/** Get pixel index in downsample texture */
 	UE_DEPRECATED(5.3, "Please use UDMXPixelMappingPixelMapRenderer to render the pixel map")
-	virtual int32 GetDownsamplePixelIndex() const { return 0; }
+		virtual int32 GetDownsamplePixelIndex() const { return 0; }
 
 	/** Queue rendering to downsample rendering target */
 	UE_DEPRECATED(5.3, "Please use UDMXPixelMappingPixelMapRenderer to render the pixel map")
-	virtual void QueueDownsample() {}
+		virtual void QueueDownsample() {}
 
-	/** Sets the position */
+	/** Returns the absolute position, without rotation */
+	FVector2D GetPosition() const;
+
+	/** Sets the absolute position of the component, without rotation. */
 	virtual void SetPosition(const FVector2D& NewPosition);
 
-	/** Returns the position */
-	FVector2D GetPosition() const;
+	/** Returns the absolute position with rotation. */
+	FVector2D GetPositionRotated() const;
+
+	/** Sets the the absolute position with rotation. */
+	virtual void SetPositionRotated(FVector2D NewRotatedPosition);
+
+	/** Get the absolute size */
+	FVector2D GetSize() const;
 
-	/** Sets the size */
+	/** Sets the absolue size */
 	virtual void SetSize(const FVector2D& NewSize);
 
-	/** Get the size */
-	FVector2D GetSize() const { return FVector2D(SizeX, SizeY); }
+	/** Gets the absolute rotation, in degrees */
+	double GetRotation() const;
+
+	/** Sets the absolute rotation, in degrees */
+	virtual void SetRotation(double NewRotation);
 
 	/** Invalidates the pixel map, effectively causing the renderer component to aquire a new pixel map */
 	void InvalidatePixelMapRenderer();
@@ -122,23 +118,10 @@ public:
 	/** Updates children to match the size of this instance */
 
 #if WITH_EDITOR
-	/** Makes the component the highest ZOrdered of components in the component rectangle, updates childs if needed */
+	/** Z-orders this component and its children topmost */
 	void ZOrderTopmost();
 #endif // WITH_EDITOR
 
-public:
-#if WITH_EDITOR
-	/** Returns the component widget */
-	UE_DEPRECATED(5.1, "Pixel Mapping Components no longer hold their own widget, in an effort to separate Views from Data.")
-	PRAGMA_DISABLE_DEPRECATION_WARNINGS // FDMXPixelMappingComponentWidget is deprecated
-	FORCEINLINE const TSharedPtr<FDMXPixelMappingComponentWidget> GetComponentWidget() { return ComponentWidget_DEPRECATED; }
-	PRAGMA_ENABLE_DEPRECATION_WARNINGS
-
-protected:
-	/** Udpates the component widget. If bWithChildrenRecursive, updates child Components' Component Widget recursively */
-	void UpdateComponentWidget(EVisibility NewVisibility = EVisibility::Visible, bool bWithChildrenRecursive = false);
-#endif // WITH_EDITOR
-
 public:
 #if WITH_EDITORONLY_DATA
 	/** ZOrder in the UI */
@@ -148,24 +131,8 @@ public:
 	/** The color displayed in editor */
 	UPROPERTY(BlueprintReadWrite, EditAnywhere, Category = "Editor Settings")
 	FLinearColor EditorColor = FLinearColor::Blue;
-
-	/** Where the component's label is shown, if it uses one. */
-	EDMXPixelMappingComponentLabelAlignment LabelAlignment;
-
-protected:
-	/** Children available PreEditUndo, useful to hide all removed ones in post edit undo */
-	TArray<UDMXPixelMappingBaseComponent*> PreEditUndoChildren;
-
-	PRAGMA_DISABLE_DEPRECATION_WARNINGS
-	/** The widget shown for this component */
-	TSharedPtr<FDMXPixelMappingComponentWidget> ComponentWidget_DEPRECATED;
-	PRAGMA_ENABLE_DEPRECATION_WARNINGS
 #endif // WITH_EDITORONLY_DATA
 
-
-	/*----------------------------------
-		Blueprint interface
-	----------------------------------*/
 public:
 #if WITH_EDITOR
 	// Property Name getters
@@ -175,23 +142,35 @@ public:
 	FORCEINLINE static FName GetPositionYPropertyName() { return GET_MEMBER_NAME_CHECKED(UDMXPixelMappingOutputComponent, PositionY); }
 	FORCEINLINE static FName GetSizeXPropertyName() { return GET_MEMBER_NAME_CHECKED(UDMXPixelMappingOutputComponent, SizeX); }
 	FORCEINLINE static FName GetSizeYPropertyName() { return GET_MEMBER_NAME_CHECKED(UDMXPixelMappingOutputComponent, SizeY); }
+	FORCEINLINE static FName GetRotationPropertyName() { return GET_MEMBER_NAME_CHECKED(UDMXPixelMappingOutputComponent, Rotation); }
 #endif
 
-private:
-	UPROPERTY(EditAnywhere, BlueprintReadOnly, Category = "Transform", Meta = (EditCondition = "!bLockInDesigner", AllowPrivateAccess = true))
-	float PositionX;
+protected:
+	UPROPERTY(BlueprintReadOnly, Category = "Transform", Meta = (AllowPrivateAccess = true))
+	float PositionX = 0.f;
 
-	UPROPERTY(EditAnywhere, BlueprintReadOnly, Category = "Transform", Meta = (EditCondition = "!bLockInDesigner", AllowPrivateAccess = true))
-	float PositionY;
+	UPROPERTY(BlueprintReadOnly, Category = "Transform", Meta = (AllowPrivateAccess = true))
+	float PositionY = 0.f;
 
-	UPROPERTY(EditAnywhere, BlueprintReadOnly, Category = "Transform", Meta = (ClampMin = 1, UIMin = 1, EditCondition = "!bLockInDesigner", AllowPrivateAccess = true))
-	float SizeX;
+#if WITH_EDITORONLY_DATA
+	UPROPERTY(EditAnywhere, Transient, Category = "Transform", Meta = (EditCondition = "!bLockInDesigner", AllowPrivateAccess = true, DisplayName = "Position (with Rotation)"))
+	FVector2D EditorPositionWithRotation = FVector2D::ZeroVector;
+#endif // WITH_EDITORONLY_DATA
 
-	UPROPERTY(EditAnywhere, BlueprintReadOnly, Category = "Transform", Meta = (ClampMin = 1, UIMin = 1, EditCondition = "!bLockInDesigner", AllowPrivateAccess = true))
-	float SizeY;
+	UPROPERTY(EditAnywhere, BlueprintReadOnly, Category = "Transform", Meta = (ClampMin = 0.0001, EditCondition = "!bLockInDesigner", AllowPrivateAccess = true))
+	float SizeX = 1.f;
+
+	UPROPERTY(EditAnywhere, BlueprintReadOnly, Category = "Transform", Meta = (ClampMin = 0.0001, EditCondition = "!bLockInDesigner", AllowPrivateAccess = true))
+	float SizeY = 1.f;
+
+	UPROPERTY(EditAnywhere, BlueprintReadOnly, Category = "Transform", Meta = (UIMin = -180, UIMax = 180, EditCondition = "!bLockInDesigner", AllowPrivateAccess = true))
+	double Rotation = 0.0;
+
+	UPROPERTY(Transient)
+	FDMXPixelMappingComponentGeometryCache CachedGeometry;
 
-protected:
 #if WITH_EDITORONLY_DATA
+protected:
 	UPROPERTY(EditAnywhere, BlueprintReadOnly, Category = "Editor Settings")
 	bool bLockInDesigner;
 
diff --git a/Engine/Plugins/VirtualProduction/DMX/DMXPixelMapping/Source/DMXPixelMappingRuntime/Public/Components/DMXPixelMappingRendererComponent.h b/Engine/Plugins/VirtualProduction/DMX/DMXPixelMapping/Source/DMXPixelMappingRuntime/Public/Components/DMXPixelMappingRendererComponent.h
--- a/Engine/Plugins/VirtualProduction/DMX/DMXPixelMapping/Source/DMXPixelMappingRuntime/Public/Components/DMXPixelMappingRendererComponent.h
+++ b/Engine/Plugins/VirtualProduction/DMX/DMXPixelMapping/Source/DMXPixelMappingRuntime/Public/Components/DMXPixelMappingRendererComponent.h
@@ -194,14 +194,6 @@ public:
 	UE_DEPRECATED(5.3, "Please use UDMXPixelMappingPixelMapRenderer to render the pixel map")
 	void EmptyDownsampleBuffer();
 
-#if WITH_EDITOR
-	UE_DEPRECATED(5.1, "Pixel Mapping Components no longer hold their own widget, in an effort to separate Views from Data.")
-	TSharedRef<SWidget> TakeWidget();
-
-	UE_DEPRECATED(5.1, "Pixel Mapping Components no longer hold their own widget, in an effort to separate Views from Data.")
-	FORCEINLINE TSharedPtr<SConstraintCanvas> GetComponentsCanvas() const { return ComponentsCanvas_DEPRECATED; }
-#endif // WITH_EDITOR
-
 private:
 	UE_DEPRECATED(5.3, "Please use UDMXPixelMappingPixelMapRenderer to render the pixel map")
 	int32 GetTotalDownsamplePixelCount();
@@ -219,10 +211,6 @@ private:
 	TSharedPtr<IDMXPixelMappingRenderer> PixelMappingRenderer_DEPRECATED;
 	PRAGMA_ENABLE_DEPRECATION_WARNINGS
 
-#if WITH_EDITORONLY_DATA
-	TSharedPtr<SConstraintCanvas> ComponentsCanvas_DEPRECATED;
-#endif // WITH_EDITORONLY_DATA
-
 	UPROPERTY(Transient, DuplicateTransient)
 	TObjectPtr<UTextureRenderTarget2D> DownsampleBufferTarget_DEPRECATED;
 
diff --git a/Engine/Plugins/VirtualProduction/DMX/DMXPixelMapping/Source/DMXPixelMappingRuntime/Public/Components/DMXPixelMappingScreenComponent.h b/Engine/Plugins/VirtualProduction/DMX/DMXPixelMapping/Source/DMXPixelMappingRuntime/Public/Components/DMXPixelMappingScreenComponent.h
--- a/Engine/Plugins/VirtualProduction/DMX/DMXPixelMapping/Source/DMXPixelMappingRuntime/Public/Components/DMXPixelMappingScreenComponent.h
+++ b/Engine/Plugins/VirtualProduction/DMX/DMXPixelMapping/Source/DMXPixelMappingRuntime/Public/Components/DMXPixelMappingScreenComponent.h
@@ -50,12 +50,7 @@ public:
 #if WITH_EDITOR
 	virtual const FText GetPaletteCategory() override;
 	virtual bool IsExposedToTemplate() { return true; }
-	PRAGMA_DISABLE_DEPRECATION_WARNINGS // FDMXPixelMappingComponentWidget is deprecated
-	virtual TSharedRef<FDMXPixelMappingComponentWidget> BuildSlot(TSharedRef<SConstraintCanvas> InCanvas) override;
-	PRAGMA_ENABLE_DEPRECATION_WARNINGS
 #endif // WITH_EDITOR
-	virtual void SetPosition(const FVector2D& NewPosition) override;
-	virtual void SetSize(const FVector2D& NewSize) override;
 	//~ End UDMXPixelMappingOutputComponent implementation
 
 public:
diff --git a/Engine/Plugins/VirtualProduction/DMX/DMXPixelMapping/Source/DMXPixelMappingRuntime/Public/DMXPixelMapping.h b/Engine/Plugins/VirtualProduction/DMX/DMXPixelMapping/Source/DMXPixelMappingRuntime/Public/DMXPixelMapping.h
--- a/Engine/Plugins/VirtualProduction/DMX/DMXPixelMapping/Source/DMXPixelMappingRuntime/Public/DMXPixelMapping.h
+++ b/Engine/Plugins/VirtualProduction/DMX/DMXPixelMapping/Source/DMXPixelMappingRuntime/Public/DMXPixelMapping.h
@@ -54,7 +54,8 @@ public:
 
 #if WITH_EDITOR
 	/** Find the component by widget. */
-	UDMXPixelMappingOutputComponent* FindComponent(TSharedPtr<SWidget> InWidget) const;
+	UE_DEPRECATED(5.4, "Component widgets are no longer supported since 5.1. This function will always return nullptr.")
+	UDMXPixelMappingOutputComponent* FindComponent(TSharedPtr<SWidget> InWidget) const { return nullptr;}
 #endif // WITH_EDITOR
 
 	/**
@@ -159,6 +160,10 @@ public:
 	UPROPERTY()
 	FLinearColor SnapGridColor;
 
+	/** If true, editor is set to scale children with parent. This is forwarded from the editor module (DMXPixelMappingEditorSettings) to be accessible in the runtime module. */
+	UPROPERTY(Transient, NonTransactional)
+	bool bEditorScaleChildrenWithParent = false;
+
 	/** Holds the Thumbnail image for this asset */
 	UPROPERTY()
 	TObjectPtr<UTexture> ThumbnailImage;
