DMX - Rework how pixel mapping components are painted, better aesthetics and editor performance

diff --git a/Engine/Plugins/VirtualProduction/DMX/DMXPixelMapping/Source/DMXPixelMappingEditor/Private/ViewModels/DMXPixelMappingOutputComponentModel.cpp b/Engine/Plugins/VirtualProduction/DMX/DMXPixelMapping/Source/DMXPixelMappingEditor/Private/ViewModels/DMXPixelMappingOutputComponentModel.cpp
--- a/Engine/Plugins/VirtualProduction/DMX/DMXPixelMapping/Source/DMXPixelMappingEditor/Private/ViewModels/DMXPixelMappingOutputComponentModel.cpp
+++ b/Engine/Plugins/VirtualProduction/DMX/DMXPixelMapping/Source/DMXPixelMappingEditor/Private/ViewModels/DMXPixelMappingOutputComponentModel.cpp
@@ -19,389 +19,402 @@
 
 #define LOCTEXT_NAMESPACE "DMXPixelMappingOutputComponentModel"
 
-FDMXPixelMappingOutputComponentModel::FDMXPixelMappingOutputComponentModel(const TSharedRef<FDMXPixelMappingToolkit>& InToolkit, TWeakObjectPtr<UDMXPixelMappingOutputComponent> InOutputComponent)
-	: WeakOutputComponent(InOutputComponent)
-	, WeakToolkit(InToolkit)
+namespace UE::DMX
 {
-	UDMXPixelMappingOutputComponent* OutputComponent = WeakOutputComponent.Get();
-	if (!OutputComponent)
+	FDMXPixelMappingOutputComponentModel::FDMXPixelMappingOutputComponentModel(const TSharedRef<FDMXPixelMappingToolkit>& InToolkit, TWeakObjectPtr<UDMXPixelMappingOutputComponent> InOutputComponent)
+		: WeakToolkit(InToolkit)
 	{
-		return;
-	}
+		if (!InOutputComponent.IsValid())
+		{
+			return;
+		}
+		OutputComponent = InOutputComponent.Get();
 
-	UpdateFixtureNode();
-	bSelected = InToolkit->GetSelectedComponents().Contains(FDMXPixelMappingComponentReference(InToolkit, OutputComponent));
+		UpdateFixtureNode();
+		bSelected = InToolkit->GetSelectedComponents().Contains(FDMXPixelMappingComponentReference(InToolkit, OutputComponent));
 
-	// Handle selection changes
-	InToolkit->GetOnSelectedComponentsChangedDelegate().AddRaw(this, &FDMXPixelMappingOutputComponentModel::OnSelectedComponentsChanged);
-}
+		// Handle selection changes
+		InToolkit->GetOnSelectedComponentsChangedDelegate().AddRaw(this, &FDMXPixelMappingOutputComponentModel::OnSelectedComponentsChanged);
+	}
 
-FDMXPixelMappingOutputComponentModel::~FDMXPixelMappingOutputComponentModel()
-{
-	if (TSharedPtr<FDMXPixelMappingToolkit> Toolkit = WeakToolkit.Pin())
+	FDMXPixelMappingOutputComponentModel::~FDMXPixelMappingOutputComponentModel()
 	{
-		Toolkit->GetOnSelectedComponentsChangedDelegate().RemoveAll(this);
+		if (TSharedPtr<FDMXPixelMappingToolkit> Toolkit = WeakToolkit.Pin())
+		{
+			Toolkit->GetOnSelectedComponentsChangedDelegate().RemoveAll(this);
+		}
 	}
-}
 
-FVector2D FDMXPixelMappingOutputComponentModel::GetPosition() const
-{
-	if (UDMXPixelMappingOutputComponent* OutputComponent = WeakOutputComponent.Get())
+	FVector2D FDMXPixelMappingOutputComponentModel::GetPosition() const
 	{
-		return OutputComponent->GetPosition();
+		if (OutputComponent)
+		{
+			return OutputComponent->GetPosition();
+		}
+		return FVector2D::ZeroVector;
 	}
-	return FVector2D::ZeroVector;
-}
 
-FVector2D FDMXPixelMappingOutputComponentModel::GetSize() const
-{
-	if (UDMXPixelMappingOutputComponent* OutputComponent = WeakOutputComponent.Get())
+	FVector2D FDMXPixelMappingOutputComponentModel::GetSize() const
 	{
-		return OutputComponent->GetSize();
+		if (OutputComponent)
+		{
+			return OutputComponent->GetSize();
+		}
+		return FVector2D::ZeroVector;
 	}
-	return FVector2D::ZeroVector;
-}
 
-FQuat2D FDMXPixelMappingOutputComponentModel::GetQuaternion() const
-{
-	if (UDMXPixelMappingOutputComponent* OutputComponent = WeakOutputComponent.Get())
+	FQuat2D FDMXPixelMappingOutputComponentModel::GetQuaternion() const
 	{
-		const FQuat2D Quat(FMath::DegreesToRadians(OutputComponent->GetRotation()));
-		return Quat;
+		if (OutputComponent)
+		{
+			const FQuat2D Quat(FMath::DegreesToRadians(OutputComponent->GetRotation()));
+			return Quat;
+		}
+		return FQuat2D();
 	}
-	return FQuat2D();
-}
 
-FText FDMXPixelMappingOutputComponentModel::GetName() const
-{
-	if (UDMXPixelMappingOutputComponent* OutputComponent = Cast<UDMXPixelMappingOutputComponent>(WeakOutputComponent.Get()))
+	FText FDMXPixelMappingOutputComponentModel::GetName() const
 	{
-		return FText::FromString(OutputComponent->GetUserName());
-	}
+		if (OutputComponent)
+		{
+			return FText::FromString(OutputComponent->GetUserName());
+		}
 
-	return FText::GetEmpty();
-}
+		return FText::GetEmpty();
+	}
 
-bool FDMXPixelMappingOutputComponentModel::ShouldDraw() const
-{
-	if (UDMXPixelMappingMatrixCellComponent* CellComponent = Cast<UDMXPixelMappingMatrixCellComponent>(WeakOutputComponent.Get()))
+	bool FDMXPixelMappingOutputComponentModel::ShouldDraw() const
 	{
-		const FDMXPixelMappingDesignerSettings& DesignerSettings = GetDefault<UDMXPixelMappingEditorSettings>()->DesignerSettings;
-		return DesignerSettings.bShowMatrixCells;
+		if (OutputComponent)
+		{
+			const bool bIsVisble = OutputComponent->IsVisible();
+			if (UDMXPixelMappingMatrixCellComponent* CellComponent = Cast<UDMXPixelMappingMatrixCellComponent>(OutputComponent))
+			{
+				const FDMXPixelMappingDesignerSettings& DesignerSettings = GetDefault<UDMXPixelMappingEditorSettings>()->DesignerSettings;
+				return bIsVisble && DesignerSettings.bShowMatrixCells;
+			}
+
+			return bIsVisble;
+		}
+
+		return true;
 	}
 
-	return true;
-}
+	bool FDMXPixelMappingOutputComponentModel::ShouldDrawName() const
+	{
+		const FDMXPixelMappingDesignerSettings& DesignerSettings = GetDefault<UDMXPixelMappingEditorSettings>()->DesignerSettings;
 
-bool FDMXPixelMappingOutputComponentModel::ShouldDrawName() const
-{
-	const FDMXPixelMappingDesignerSettings& DesignerSettings = GetDefault<UDMXPixelMappingEditorSettings>()->DesignerSettings;
+		return DesignerSettings.bShowComponentNames &&
+			OutputComponent &&
+			OutputComponent->GetClass() != UDMXPixelMappingMatrixCellComponent::StaticClass();
+	}
 
-	return DesignerSettings.bShowComponentNames &&
-		WeakOutputComponent.IsValid() && 
-		WeakOutputComponent->GetClass() != UDMXPixelMappingMatrixCellComponent::StaticClass();
-}
+	bool FDMXPixelMappingOutputComponentModel::ShouldDrawNameAbove() const
+	{
+		if (OutputComponent)
+		{
+			return
+				OutputComponent->GetClass() == UDMXPixelMappingFixtureGroupComponent::StaticClass() ||
+				OutputComponent->GetClass() == UDMXPixelMappingScreenComponent::StaticClass();
+		}
+		return true;
+	}
 
-bool FDMXPixelMappingOutputComponentModel::ShouldDrawNameAbove() const
-{
-	if (UDMXPixelMappingOutputComponent* OutputComponent = Cast<UDMXPixelMappingOutputComponent>(WeakOutputComponent.Get()))
+	bool FDMXPixelMappingOutputComponentModel::ShouldDrawCellID() const
 	{
+		const FDMXPixelMappingDesignerSettings& DesignerSettings = GetDefault<UDMXPixelMappingEditorSettings>()->DesignerSettings;
+
 		return
-			OutputComponent->GetClass() == UDMXPixelMappingFixtureGroupComponent::StaticClass() ||
-			OutputComponent->GetClass() == UDMXPixelMappingScreenComponent::StaticClass();
+			DesignerSettings.bShowCellIDs &&
+			OutputComponent &&
+			OutputComponent->GetClass() == UDMXPixelMappingMatrixCellComponent::StaticClass();
 	}
-	return true;
-}
 
-bool FDMXPixelMappingOutputComponentModel::ShouldDrawCellID() const
-{
-	const FDMXPixelMappingDesignerSettings& DesignerSettings = GetDefault<UDMXPixelMappingEditorSettings>()->DesignerSettings;
+	bool FDMXPixelMappingOutputComponentModel::ShouldDrawPatchInfo() const
+	{
+		const FDMXPixelMappingDesignerSettings& DesignerSettings = GetDefault<UDMXPixelMappingEditorSettings>()->DesignerSettings;
 
-	return
-		DesignerSettings.bShowCellIDs &&
-		WeakOutputComponent.IsValid() &&
-		WeakOutputComponent->GetClass() == UDMXPixelMappingMatrixCellComponent::StaticClass();
-}
+		return
+			DesignerSettings.bShowPatchInfo &&
+			GetFixturePatch() != nullptr;;
+	}
 
-bool FDMXPixelMappingOutputComponentModel::ShouldDrawPatchInfo() const
-{
-	const FDMXPixelMappingDesignerSettings& DesignerSettings = GetDefault<UDMXPixelMappingEditorSettings>()->DesignerSettings;
+	FText FDMXPixelMappingOutputComponentModel::GetAddressesText() const
+	{
+		if (const UDMXEntityFixturePatch* FixturePatch = GetFixturePatch())
+		{
+			const FString AddressesString = FString::Printf(TEXT("%i.%i"), FixturePatch->GetUniverseID(), FixturePatch->GetStartingChannel());
 
-	return
-		DesignerSettings.bShowPatchInfo &&
-		GetFixturePatch() != nullptr;;
-}
+			return FText::Format(LOCTEXT("PatchTextWithLabel", "Patch: {0}"), FText::FromString(AddressesString));
+		}
 
-FText FDMXPixelMappingOutputComponentModel::GetAddressesText() const
-{
-	if (const UDMXEntityFixturePatch* FixturePatch = GetFixturePatch())
-	{
-		const FString AddressesString = FString::Printf(TEXT("%i.%i"), FixturePatch->GetUniverseID(), FixturePatch->GetStartingChannel());
-		
-		return FText::Format(LOCTEXT("PatchTextWithLabel", "Patch: {0}"), FText::FromString(AddressesString));
+		return FText::GetEmpty();
 	}
 
-	return FText::GetEmpty();
-}
-
-FText FDMXPixelMappingOutputComponentModel::GetFixtureIDText() const
-{
-	if (UDMXMVRFixtureNode* FixtureNode = WeakFixtureNode.Get())
+	FText FDMXPixelMappingOutputComponentModel::GetFixtureIDText() const
 	{
-		const FText FixtureIDText = FixtureNode->FixtureID.IsEmpty() ? LOCTEXT("NoFixtureID", "<not set>") : FText::FromString(FixtureNode->FixtureID);
-		return FText::Format(LOCTEXT("FixtureIDTextWithLabel", "FID: {0}"), FixtureIDText);
-	}
+		if (UDMXMVRFixtureNode* FixtureNode = WeakFixtureNode.Get())
+		{
+			const FText FixtureIDText = FixtureNode->FixtureID.IsEmpty() ? LOCTEXT("NoFixtureID", "<not set>") : FText::FromString(FixtureNode->FixtureID);
+			return FText::Format(LOCTEXT("FixtureIDTextWithLabel", "FID: {0}"), FixtureIDText);
+		}
 
-	return FText::GetEmpty();
-}
+		return FText::GetEmpty();
+	}
 
-FText FDMXPixelMappingOutputComponentModel::GetCellIDText() const
-{
-	if (UDMXPixelMappingMatrixCellComponent* MatrixCell = Cast<UDMXPixelMappingMatrixCellComponent>(WeakOutputComponent.Get()))
+	FText FDMXPixelMappingOutputComponentModel::GetCellIDText() const
 	{
-		return FText::FromString(FString::FromInt(MatrixCell->CellID));
-	}
+		if (UDMXPixelMappingMatrixCellComponent* MatrixCell = Cast<UDMXPixelMappingMatrixCellComponent>(OutputComponent))
+		{
+			return FText::FromString(FString::FromInt(MatrixCell->CellID));
+		}
 
-	return FText::GetEmpty();
-}
+		return FText::GetEmpty();
+	}
 
-FLinearColor FDMXPixelMappingOutputComponentModel::GetColor() const
-{
-	if (UDMXPixelMappingOutputComponent* OutputComponent = WeakOutputComponent.Get())
+	FLinearColor FDMXPixelMappingOutputComponentModel::GetColor() const
 	{
-		if (OutputComponent->IsOverParent())
+		if (OutputComponent)
 		{
-			if (bSelected)
-			{
-				return OutputComponent->GetEditorColor();
-			}
-			else
+			if (OutputComponent->IsOverParent())
 			{
-				const float Opacity = OutputComponent->GetEditorColor().A * .75f;
-				return OutputComponent->GetEditorColor().CopyWithNewOpacity(Opacity);
+				if (bSelected)
+				{
+					return OutputComponent->GetEditorColor();
+				}
+				else
+				{
+					const float Opacity = OutputComponent->GetEditorColor().A * .6f;
+					return OutputComponent->GetEditorColor().CopyWithNewOpacity(Opacity);
+				}
 			}
 		}
-	}
-
-	return FLinearColor::Red;
-}
-
-bool FDMXPixelMappingOutputComponentModel::Equals(UDMXPixelMappingBaseComponent* Other) const
-{
-	return WeakOutputComponent.Get() == Other; 
-}
 
-void FDMXPixelMappingOutputComponentModel::OnSelectedComponentsChanged()
-{
-	TSharedPtr<FDMXPixelMappingToolkit> Toolkit = WeakToolkit.Pin();
-	if (!Toolkit.IsValid())
-	{
-		return;
+		return FLinearColor::Red;
 	}
 
-	if (UDMXPixelMappingOutputComponent* OutputComponent = WeakOutputComponent.Get())
+	bool FDMXPixelMappingOutputComponentModel::Equals(UDMXPixelMappingBaseComponent* Other) const
 	{
-		bSelected = Toolkit->GetSelectedComponents().Contains(FDMXPixelMappingComponentReference(Toolkit, OutputComponent));
+		return OutputComponent && OutputComponent == Other;
 	}
-}
 
-const UDMXEntityFixturePatch* FDMXPixelMappingOutputComponentModel::GetFixturePatch() const
-{
-	const UDMXEntityFixturePatch* FixturePatch = nullptr;
-	if (UDMXPixelMappingMatrixComponent* Matrix = Cast<UDMXPixelMappingMatrixComponent>(WeakOutputComponent.Get()))
+	void FDMXPixelMappingOutputComponentModel::AddReferencedObjects(FReferenceCollector& Collector)
 	{
-		FixturePatch = Matrix->FixturePatchRef.GetFixturePatch();
+		Collector.AddReferencedObject(OutputComponent);
 	}
-	else if (UDMXPixelMappingFixtureGroupItemComponent* GroupItem = Cast<UDMXPixelMappingFixtureGroupItemComponent>(WeakOutputComponent.Get()))
-	{
-		FixturePatch = GroupItem->FixturePatchRef.GetFixturePatch();
-	}
-
-	return FixturePatch;
-}
-
-void FDMXPixelMappingOutputComponentModel::UpdateFixtureNode()
-{
-	WeakFixtureNode.Reset();
 
-	const UDMXEntityFixturePatch* FixturePatch = GetFixturePatch();
-	if (!FixturePatch)
+	void FDMXPixelMappingOutputComponentModel::OnSelectedComponentsChanged()
 	{
-		return;
-	}
+		TSharedPtr<FDMXPixelMappingToolkit> Toolkit = WeakToolkit.Pin();
+		if (!Toolkit.IsValid())
+		{
+			return;
+		}
 
-	const FGuid& MVRUUID = FixturePatch->GetMVRFixtureUUID();
-	if (!MVRUUID.IsValid())
-	{
-		return;
+		if (OutputComponent)
+		{
+			bSelected = Toolkit->GetSelectedComponents().Contains(FDMXPixelMappingComponentReference(Toolkit, OutputComponent));
+		}
 	}
 
-	const UDMXLibrary* DMXLibrary = FixturePatch->GetParentLibrary();
-	if (!DMXLibrary)
+	const UDMXEntityFixturePatch* FDMXPixelMappingOutputComponentModel::GetFixturePatch() const
 	{
-		return;
+		const UDMXEntityFixturePatch* FixturePatch = nullptr;
+		if (UDMXPixelMappingMatrixComponent* Matrix = Cast<UDMXPixelMappingMatrixComponent>(OutputComponent))
+		{
+			FixturePatch = Matrix->FixturePatchRef.GetFixturePatch();
+		}
+		else if (UDMXPixelMappingFixtureGroupItemComponent* GroupItem = Cast<UDMXPixelMappingFixtureGroupItemComponent>(OutputComponent))
+		{
+			FixturePatch = GroupItem->FixturePatchRef.GetFixturePatch();
+		}
+
+		return FixturePatch;
 	}
 
-	// We can use the lazy General Scene Description
-	UDMXMVRGeneralSceneDescription* GeneralSceneDescription = DMXLibrary->GetLazyGeneralSceneDescription();
-	if (!GeneralSceneDescription)
+	void FDMXPixelMappingOutputComponentModel::UpdateFixtureNode()
 	{
-		return;
-	}
+		WeakFixtureNode.Reset();
 
-	WeakFixtureNode = GeneralSceneDescription->FindFixtureNode(MVRUUID);
-}
+		const UDMXEntityFixturePatch* FixturePatch = GetFixturePatch();
+		if (!FixturePatch)
+		{
+			return;
+		}
 
-FDMXPixelMappingScreenComponentModel::FDMXPixelMappingScreenComponentModel(const TSharedRef<FDMXPixelMappingToolkit>& InToolkit, TWeakObjectPtr<UDMXPixelMappingScreenComponent> InScreenComponent)
-	: WeakScreenComponent(InScreenComponent)
-	, WeakToolkit(InToolkit)
-{
-	UDMXPixelMappingScreenComponent* ScreenComponent = WeakScreenComponent.Get();
-	if (!ScreenComponent)
-	{
-		return;
-	}
+		const FGuid& MVRUUID = FixturePatch->GetMVRFixtureUUID();
+		if (!MVRUUID.IsValid())
+		{
+			return;
+		}
 
-	bSelected = InToolkit->GetSelectedComponents().Contains(FDMXPixelMappingComponentReference(InToolkit, ScreenComponent));
+		const UDMXLibrary* DMXLibrary = FixturePatch->GetParentLibrary();
+		if (!DMXLibrary)
+		{
+			return;
+		}
 
-	// Handle selection changes
-	InToolkit->GetOnSelectedComponentsChangedDelegate().AddRaw(this, &FDMXPixelMappingScreenComponentModel::OnSelectedComponentsChanged);
-}
+		// We can use the lazy General Scene Description
+		UDMXMVRGeneralSceneDescription* GeneralSceneDescription = DMXLibrary->GetLazyGeneralSceneDescription();
+		if (!GeneralSceneDescription)
+		{
+			return;
+		}
 
-FDMXPixelMappingScreenComponentModel::~FDMXPixelMappingScreenComponentModel()
-{
-	if (TSharedPtr<FDMXPixelMappingToolkit> Toolkit = WeakToolkit.Pin())
-	{
-		Toolkit->GetOnSelectedComponentsChangedDelegate().RemoveAll(this);
+		WeakFixtureNode = GeneralSceneDescription->FindFixtureNode(MVRUUID);
 	}
-}
 
-FVector2D FDMXPixelMappingScreenComponentModel::GetPosition() const
-{
-	if (UDMXPixelMappingScreenComponent* ScreenComponent = WeakScreenComponent.Get())
+	FDMXPixelMappingScreenComponentModel::FDMXPixelMappingScreenComponentModel(const TSharedRef<FDMXPixelMappingToolkit>& InToolkit, TWeakObjectPtr<UDMXPixelMappingScreenComponent> InScreenComponent)
+		: WeakScreenComponent(InScreenComponent)
+		, WeakToolkit(InToolkit)
 	{
-		return ScreenComponent->GetPosition();
+		UDMXPixelMappingScreenComponent* ScreenComponent = WeakScreenComponent.Get();
+		if (!ScreenComponent)
+		{
+			return;
+		}
+
+		bSelected = InToolkit->GetSelectedComponents().Contains(FDMXPixelMappingComponentReference(InToolkit, ScreenComponent));
+
+		// Handle selection changes
+		InToolkit->GetOnSelectedComponentsChangedDelegate().AddRaw(this, &FDMXPixelMappingScreenComponentModel::OnSelectedComponentsChanged);
 	}
-	return FVector2D::ZeroVector;
-}
 
-FVector2D FDMXPixelMappingScreenComponentModel::GetSize() const
-{
-	if (UDMXPixelMappingScreenComponent* ScreenComponent = WeakScreenComponent.Get())
+	FDMXPixelMappingScreenComponentModel::~FDMXPixelMappingScreenComponentModel()
 	{
-		return ScreenComponent->GetSize();
+		if (TSharedPtr<FDMXPixelMappingToolkit> Toolkit = WeakToolkit.Pin())
+		{
+			Toolkit->GetOnSelectedComponentsChangedDelegate().RemoveAll(this);
+		}
 	}
-	return FVector2D::ZeroVector;
-}
 
-int32 FDMXPixelMappingScreenComponentModel::GetNumColumns() const
-{
-	if (UDMXPixelMappingScreenComponent* ScreenComponent = WeakScreenComponent.Get())
+	FVector2D FDMXPixelMappingScreenComponentModel::GetPosition() const
 	{
-		return ScreenComponent->NumXCells;
+		if (UDMXPixelMappingScreenComponent* ScreenComponent = WeakScreenComponent.Get())
+		{
+			return ScreenComponent->GetPosition();
+		}
+		return FVector2D::ZeroVector;
 	}
-	return 1;
-}
 
-int32 FDMXPixelMappingScreenComponentModel::GetNumRows() const
-{
-	if (UDMXPixelMappingScreenComponent* ScreenComponent = WeakScreenComponent.Get())
+	FVector2D FDMXPixelMappingScreenComponentModel::GetSize() const
 	{
-		return ScreenComponent->NumYCells;
+		if (UDMXPixelMappingScreenComponent* ScreenComponent = WeakScreenComponent.Get())
+		{
+			return ScreenComponent->GetSize();
+		}
+		return FVector2D::ZeroVector;
 	}
-	return 1;
-}
 
-FLinearColor FDMXPixelMappingScreenComponentModel::GetColor() const
-{
-	if (UDMXPixelMappingScreenComponent* ScreenComponent = WeakScreenComponent.Get())
+	int32 FDMXPixelMappingScreenComponentModel::GetNumColumns() const
 	{
-		if (bSelected)
+		if (UDMXPixelMappingScreenComponent* ScreenComponent = WeakScreenComponent.Get())
 		{
-			return ScreenComponent->GetEditorColor();
+			return ScreenComponent->NumXCells;
 		}
-		else
+		return 1;
+	}
+
+	int32 FDMXPixelMappingScreenComponentModel::GetNumRows() const
+	{
+		if (UDMXPixelMappingScreenComponent* ScreenComponent = WeakScreenComponent.Get())
 		{
-			return ScreenComponent->GetEditorColor().CopyWithNewOpacity(0.8f);
+			return ScreenComponent->NumYCells;
 		}
+		return 1;
 	}
 
-	return FLinearColor::Red;
-}
+	FLinearColor FDMXPixelMappingScreenComponentModel::GetColor() const
+	{
+		if (UDMXPixelMappingScreenComponent* ScreenComponent = WeakScreenComponent.Get())
+		{
+			if (bSelected)
+			{
+				return ScreenComponent->GetEditorColor();
+			}
+			else
+			{
+				return ScreenComponent->GetEditorColor().CopyWithNewOpacity(0.6f);
+			}
+		}
 
-bool FDMXPixelMappingScreenComponentModel::Equals(UDMXPixelMappingBaseComponent* Other) const
-{
-	return WeakScreenComponent.Get() == Other;
-}
+		return FLinearColor::Red;
+	}
 
-void FDMXPixelMappingScreenComponentModel::OnSelectedComponentsChanged()
-{
-	TSharedPtr<FDMXPixelMappingToolkit> Toolkit = WeakToolkit.Pin();
-	if (!Toolkit.IsValid())
+	bool FDMXPixelMappingScreenComponentModel::Equals(UDMXPixelMappingBaseComponent* Other) const
 	{
-		return;
+		return WeakScreenComponent.Get() == Other;
 	}
 
-	if (UDMXPixelMappingScreenComponent* ScreenComponent = WeakScreenComponent.Get())
+	void FDMXPixelMappingScreenComponentModel::OnSelectedComponentsChanged()
 	{
-		bSelected = Toolkit->GetSelectedComponents().Contains(FDMXPixelMappingComponentReference(Toolkit, ScreenComponent));
+		TSharedPtr<FDMXPixelMappingToolkit> Toolkit = WeakToolkit.Pin();
+		if (!Toolkit.IsValid())
+		{
+			return;
+		}
+
+		if (UDMXPixelMappingScreenComponent* ScreenComponent = WeakScreenComponent.Get())
+		{
+			bSelected = Toolkit->GetSelectedComponents().Contains(FDMXPixelMappingComponentReference(Toolkit, ScreenComponent));
+		}
 	}
-}
 
-EDMXPixelMappingDistribution FDMXPixelMappingScreenComponentModel::GetDistribution() const
-{
-	if (UDMXPixelMappingScreenComponent* ScreenComponent = WeakScreenComponent.Get())
+	EDMXPixelMappingDistribution FDMXPixelMappingScreenComponentModel::GetDistribution() const
 	{
-		return ScreenComponent->Distribution;
+		if (UDMXPixelMappingScreenComponent* ScreenComponent = WeakScreenComponent.Get())
+		{
+			return ScreenComponent->Distribution;
+		}
+		return EDMXPixelMappingDistribution::BottomLeftToRight;
 	}
-	return EDMXPixelMappingDistribution::BottomLeftToRight;
-}
 
-EDMXCellFormat FDMXPixelMappingScreenComponentModel::GetCellFormat() const
-{
-	if (UDMXPixelMappingScreenComponent* ScreenComponent = WeakScreenComponent.Get())
+	EDMXCellFormat FDMXPixelMappingScreenComponentModel::GetCellFormat() const
 	{
-		return ScreenComponent->PixelFormat;
+		if (UDMXPixelMappingScreenComponent* ScreenComponent = WeakScreenComponent.Get())
+		{
+			return ScreenComponent->PixelFormat;
+		}
+		return EDMXCellFormat::PF_RGB;
 	}
-	return EDMXCellFormat::PF_RGB;
-}
 
-bool FDMXPixelMappingScreenComponentModel::ComponentWantsToShowUniverse() const
-{
-	if (UDMXPixelMappingScreenComponent* ScreenComponent = WeakScreenComponent.Get())
+	bool FDMXPixelMappingScreenComponentModel::ComponentWantsToShowUniverse() const
 	{
-		return ScreenComponent->bShowUniverse;
+		if (UDMXPixelMappingScreenComponent* ScreenComponent = WeakScreenComponent.Get())
+		{
+			return ScreenComponent->bShowUniverse;
+		}
+		return false;
 	}
-	return false;
-}
 
-int32 FDMXPixelMappingScreenComponentModel::GetUniverse() const
-{
-	if (UDMXPixelMappingScreenComponent* ScreenComponent = WeakScreenComponent.Get())
+	int32 FDMXPixelMappingScreenComponentModel::GetUniverse() const
 	{
-		return ScreenComponent->LocalUniverse;
-	}
+		if (UDMXPixelMappingScreenComponent* ScreenComponent = WeakScreenComponent.Get())
+		{
+			return ScreenComponent->LocalUniverse;
+		}
 
-	return 0;
-}
+		return 0;
+	}
 
-bool FDMXPixelMappingScreenComponentModel::ComponentWantsToShowChannel() const
-{
-	if (UDMXPixelMappingScreenComponent* ScreenComponent = WeakScreenComponent.Get())
+	bool FDMXPixelMappingScreenComponentModel::ComponentWantsToShowChannel() const
 	{
-		return ScreenComponent->bShowAddresses;
+		if (UDMXPixelMappingScreenComponent* ScreenComponent = WeakScreenComponent.Get())
+		{
+			return ScreenComponent->bShowAddresses;
+		}
+		return false;
 	}
-	return false;
-}
 
-int32 FDMXPixelMappingScreenComponentModel::GetStartingChannel() const
-{
-	if (UDMXPixelMappingScreenComponent* ScreenComponent = WeakScreenComponent.Get())
+	int32 FDMXPixelMappingScreenComponentModel::GetStartingChannel() const
 	{
-		return ScreenComponent->StartAddress;
-	}
+		if (UDMXPixelMappingScreenComponent* ScreenComponent = WeakScreenComponent.Get())
+		{
+			return ScreenComponent->StartAddress;
+		}
 
-	return 0;
+		return 0;
+	}
 }
 
 #undef LOCTEXT_NAMESPACE
diff --git a/Engine/Plugins/VirtualProduction/DMX/DMXPixelMapping/Source/DMXPixelMappingEditor/Private/ViewModels/DMXPixelMappingOutputComponentModel.h b/Engine/Plugins/VirtualProduction/DMX/DMXPixelMapping/Source/DMXPixelMappingEditor/Private/ViewModels/DMXPixelMappingOutputComponentModel.h
--- a/Engine/Plugins/VirtualProduction/DMX/DMXPixelMapping/Source/DMXPixelMappingEditor/Private/ViewModels/DMXPixelMappingOutputComponentModel.h
+++ b/Engine/Plugins/VirtualProduction/DMX/DMXPixelMapping/Source/DMXPixelMappingEditor/Private/ViewModels/DMXPixelMappingOutputComponentModel.h
@@ -4,6 +4,7 @@
 
 #include "Math/TransformCalculus2D.h"
 #include "Math/Vector2D.h"
+#include "UObject/GCObject.h"
 #include "UObject/WeakObjectPtr.h"
 
 enum class EDMXPixelMappingDistribution : uint8;
@@ -16,143 +17,156 @@ class UDMXPixelMappingOutputComponent;
 class UDMXPixelMappingScreenComponent;
 
 
-/** 
- * Model for the Pixel Mapping Output Components.
- * For Screen Component, see DMXPixelMappingScreenComponentModel.
- */
-class FDMXPixelMappingOutputComponentModel
-	: public TSharedFromThis<FDMXPixelMappingOutputComponentModel>
+namespace UE::DMX
 {
-public:
-	/** Constructor */
-	FDMXPixelMappingOutputComponentModel(const TSharedRef<FDMXPixelMappingToolkit>& InToolkit, TWeakObjectPtr<UDMXPixelMappingOutputComponent> InOutputComponent);
+	/**
+	 * Model for the Pixel Mapping Output Components.
+	 * For Screen Component, see DMXPixelMappingScreenComponentModel.
+	 */
+	class FDMXPixelMappingOutputComponentModel
+		: public FGCObject
+		, public TSharedFromThis<FDMXPixelMappingOutputComponentModel>
+	{
+	public:
+		/** Constructor */
+		FDMXPixelMappingOutputComponentModel(const TSharedRef<FDMXPixelMappingToolkit>& InToolkit, TWeakObjectPtr<UDMXPixelMappingOutputComponent> InOutputComponent);
 
-	/** Destructor */
-	virtual ~FDMXPixelMappingOutputComponentModel();
+		/** Destructor */
+		virtual ~FDMXPixelMappingOutputComponentModel();
 
-	/** Returns the position of the component */
-	FVector2D GetPosition() const;
+		/** Returns the position of the component */
+		FVector2D GetPosition() const;
 
-	/** Returns the size of the component */
-	FVector2D GetSize() const;
+		/** Returns the size of the component */
+		FVector2D GetSize() const;
 
-	/** Returns the rotation, in degrees */
-	FQuat2D GetQuaternion() const;
+		/** Returns the rotation, in degrees */
+		FQuat2D GetQuaternion() const;
+		
+		/** Returns the Name of the component */
+		FText GetName() const;
 
-	/** Returns the Name of the component */
-	FText GetName() const;
+		/** Returns true if a widget ever should be drawn */
+		bool ShouldDraw() const;
 
-	/** Returns true if a widget ever should be drawn */
-	bool ShouldDraw() const;
+		/** Returns true if the name should drawn its name */
+		bool ShouldDrawName() const;
 
-	/** Returns true if the name should drawn its name */
-	bool ShouldDrawName() const;
+		/** Returns true if the name should drawn its name above the widget */
+		bool ShouldDrawNameAbove() const;
 
-	/** Returns true if the name should drawn its name above the widget */
-	bool ShouldDrawNameAbove() const;
+		/** Returns true if the cell ID should be shown */
+		bool ShouldDrawCellID() const;
 
-	/** Returns true if the cell ID should be shown */
-	bool ShouldDrawCellID() const;
+		/** Returns true if the patch inifo should be shown */
+		bool ShouldDrawPatchInfo() const;
 
-	/** Returns true if the patch inifo should be shown */
-	bool ShouldDrawPatchInfo() const;
+		/** Returns the addresses of the patch as text */
+		FText GetAddressesText() const;
 
-	/** Returns the addresses of the patch as text */
-	FText GetAddressesText() const;
+		/** Returns the Fixture ID as text. Empty unless a Matrix and Fixture Group Item component */
+		FText GetFixtureIDText() const;
 
-	/** Returns the Fixture ID as text. Empty unless a Matrix and Fixture Group Item component */
-	FText GetFixtureIDText() const;
+		/** Returns the Cell ID as text. Empty unless a Matrix Cell component. */
+		FText GetCellIDText() const;
 
-	/** Returns the Cell ID as text. Empty unless a Matrix Cell component. */
-	FText GetCellIDText() const;
+		/** Returns the color of the widget */
+		FLinearColor GetColor() const;
 
-	/** Returns the color of the widget */
-	FLinearColor GetColor() const;
+		/** Returns true if this Model handles the Other component */
+		bool Equals(UDMXPixelMappingBaseComponent* Other) const;
 
-	/** Returns true if this Model handles the Other component */
-	bool Equals(UDMXPixelMappingBaseComponent* Other) const; 
+	protected:
+		//~ Begin FGCObject interface
+		virtual void AddReferencedObjects(FReferenceCollector& Collector) override;
+		virtual FString GetReferencerName() const override
+		{
+			return TEXT("FDMXPixelMappingOutputComponentModel");
+		}
+		//~ End FGCObject interface
 
-private:
-	/** Called when selected Components changed */
-	void OnSelectedComponentsChanged();
+	private:
+		/** Called when selected Components changed */
+		void OnSelectedComponentsChanged();
 
-	/** Returns the Fixture Patch of the component, or nullptr if the component does not use a Fixture Patch. */
-	const UDMXEntityFixturePatch* GetFixturePatch() const;
+		/** Returns the Fixture Patch of the component, or nullptr if the component does not use a Fixture Patch. */
+		const UDMXEntityFixturePatch* GetFixturePatch() const;
 
-	/** Updates the Fixture Node that correlates to this component */
-	void UpdateFixtureNode();
+		/** Updates the Fixture Node that correlates to this component */
+		void UpdateFixtureNode();
 
-	/** The Fixture Node that corresponds to this component, only valid for Group Item and Matrix */
-	TWeakObjectPtr<UDMXMVRFixtureNode> WeakFixtureNode;
+		/** The Fixture Node that corresponds to this component, only valid for Group Item and Matrix */
+		TWeakObjectPtr<UDMXMVRFixtureNode> WeakFixtureNode;
 
-	/** True if the component is selected */
-	bool bSelected = false;
+		/** True if the component is selected */
+		bool bSelected = false;
 
-	/** The output component */
-	TWeakObjectPtr<UDMXPixelMappingOutputComponent> WeakOutputComponent;
+		/** The actual output component */
+		TObjectPtr<UDMXPixelMappingOutputComponent> OutputComponent;
 
-	/** The toolkit from which the Model currently sources */
-	TWeakPtr<FDMXPixelMappingToolkit> WeakToolkit;
-};
+		/** The toolkit from which the Model currently sources */
+		TWeakPtr<FDMXPixelMappingToolkit> WeakToolkit;
+	};
 
 
-/** Model for Pixel Mapping Screen Components */
-class FDMXPixelMappingScreenComponentModel
-	: public TSharedFromThis<FDMXPixelMappingScreenComponentModel>
-{
-public:
-	/** Constructor */
-	FDMXPixelMappingScreenComponentModel(const TSharedRef<FDMXPixelMappingToolkit>& InToolkit, TWeakObjectPtr<UDMXPixelMappingScreenComponent> InScreenComponent);
+	/** Model for Pixel Mapping Screen Components */
+	class FDMXPixelMappingScreenComponentModel
+		: public TSharedFromThis<FDMXPixelMappingScreenComponentModel>
+	{
+	public:
+		/** Constructor */
+		FDMXPixelMappingScreenComponentModel(const TSharedRef<FDMXPixelMappingToolkit>& InToolkit, TWeakObjectPtr<UDMXPixelMappingScreenComponent> InScreenComponent);
 
-	/** Destructor */
-	virtual ~FDMXPixelMappingScreenComponentModel();
+		/** Destructor */
+		virtual ~FDMXPixelMappingScreenComponentModel();
 
-	/** Returns the position of the component */
-	FVector2D GetPosition() const;
+		/** Returns the position of the component */
+		FVector2D GetPosition() const;
 
-	/** Returns the size of the component */
-	FVector2D GetSize() const;
+		/** Returns the size of the component */
+		FVector2D GetSize() const;
 
-	/** Returns the number of Columns of the Screen Component */
-	int32 GetNumColumns() const;
+		/** Returns the number of Columns of the Screen Component */
+		int32 GetNumColumns() const;
 
-	/** Returns the number of Rows of the Screen Component */
-	int32 GetNumRows() const;
+		/** Returns the number of Rows of the Screen Component */
+		int32 GetNumRows() const;
 
-	/** Returns the pixel mapping distribution of the cells */
-	EDMXPixelMappingDistribution GetDistribution() const;
+		/** Returns the pixel mapping distribution of the cells */
+		EDMXPixelMappingDistribution GetDistribution() const;
 
-	/** Returns the color format of the cells */
-	EDMXCellFormat GetCellFormat() const;
+		/** Returns the color format of the cells */
+		EDMXCellFormat GetCellFormat() const;
 
-	/** Returns true if the components wants to show the universe (it has a related property) */
-	bool ComponentWantsToShowUniverse() const;
+		/** Returns true if the components wants to show the universe (it has a related property) */
+		bool ComponentWantsToShowUniverse() const;
 
-	/** Returns the universe of the screen component */
-	int32 GetUniverse() const;
+		/** Returns the universe of the screen component */
+		int32 GetUniverse() const;
 
-	/** Returns true if the components wants to show the channel (it has a related property) */
-	bool ComponentWantsToShowChannel() const;
+		/** Returns true if the components wants to show the channel (it has a related property) */
+		bool ComponentWantsToShowChannel() const;
 
-	/** Returns the starting address of the screen component */
-	int32 GetStartingChannel() const;
+		/** Returns the starting address of the screen component */
+		int32 GetStartingChannel() const;
 
-	/** Returns the color of the widget */
-	FLinearColor GetColor() const;
+		/** Returns the color of the widget */
+		FLinearColor GetColor() const;
 
-	/** Returns true if this Model handles the Other component */
-	bool Equals(UDMXPixelMappingBaseComponent* Other) const;
+		/** Returns true if this Model handles the Other component */
+		bool Equals(UDMXPixelMappingBaseComponent* Other) const;
 
-private:
-	/** Called when selected Components changed */
-	void OnSelectedComponentsChanged();
+	private:
+		/** Called when selected Components changed */
+		void OnSelectedComponentsChanged();
 
-	/** True if the component is selected */
-	bool bSelected = false;
+		/** True if the component is selected */
+		bool bSelected = false;
 
-	/** The output component */
-	TWeakObjectPtr<UDMXPixelMappingScreenComponent> WeakScreenComponent;
+		/** The output component */
+		TWeakObjectPtr<UDMXPixelMappingScreenComponent> WeakScreenComponent;
 
-	/** The toolkit from which the Model currently sources */
-	TWeakPtr<FDMXPixelMappingToolkit> WeakToolkit;
-};
+		/** The toolkit from which the Model currently sources */
+		TWeakPtr<FDMXPixelMappingToolkit> WeakToolkit;
+	};
+}
diff --git a/Engine/Plugins/VirtualProduction/DMX/DMXPixelMapping/Source/DMXPixelMappingEditor/Private/Views/SDMXPixelMappingDesignerView.cpp b/Engine/Plugins/VirtualProduction/DMX/DMXPixelMapping/Source/DMXPixelMappingEditor/Private/Views/SDMXPixelMappingDesignerView.cpp
--- a/Engine/Plugins/VirtualProduction/DMX/DMXPixelMapping/Source/DMXPixelMappingEditor/Private/Views/SDMXPixelMappingDesignerView.cpp
+++ b/Engine/Plugins/VirtualProduction/DMX/DMXPixelMapping/Source/DMXPixelMappingEditor/Private/Views/SDMXPixelMappingDesignerView.cpp
@@ -371,9 +371,8 @@ FReply SDMXPixelMappingDesignerView::OnMouseButtonDown(const FGeometry& MyGeomet
 {
 	SDMXPixelMappingSurface::OnMouseButtonDown(MyGeometry, MouseEvent);
 
-	const bool bLeftMouseButton = MouseEvent.GetEffectingButton() == EKeys::LeftMouseButton;
 	UDMXPixelMappingOutputComponent* ClickedComponent = GetComponentUnderCursor();
-	if (bLeftMouseButton && ClickedComponent)
+	if (MouseEvent.GetEffectingButton() == EKeys::LeftMouseButton && ClickedComponent)
 	{
 		PendingSelectedComponent = ClickedComponent;
 
@@ -681,6 +680,7 @@ void SDMXPixelMappingDesignerView::RebuildDesigner()
 
 		CachedRendererComponent->ForEachChild([this, &Toolkit, &ComponentCanvas](UDMXPixelMappingBaseComponent* Component)
 			{
+				using namespace UE::DMX;
 				if (UDMXPixelMappingScreenComponent* ScreenComponent = Cast<UDMXPixelMappingScreenComponent>(Component))
 				{
 					const TSharedRef<SDMXPixelMappingScreenComponent> ScreenComponentWidget = SNew(SDMXPixelMappingScreenComponent, Toolkit.ToSharedRef(), ScreenComponent);
@@ -768,51 +768,36 @@ FVector2D SDMXPixelMappingDesignerView::GetExtensionPosition(TSharedPtr<SDMXPixe
 
 		if (GetComponentGeometry(OutputComponent, SelectedComponentGeometry))
 		{
-			const FVector2f LocalPosition = [Handle, &SelectedComponentGeometry]()
+			const FVector2D LocalPosition = [Handle, &SelectedComponentGeometry]() -> FVector2D
 			{
 				// Get the initial offset based on the location around the selected object.
 				switch (Handle->GetTransformDirection())
 				{
 				case EDMXPixelMappingTransformDirection::CenterRight:
-					return FVector2f(SelectedComponentGeometry.GetLocalSize().X, SelectedComponentGeometry.GetLocalSize().Y * 0.5f);
+					return FVector2D(SelectedComponentGeometry.GetLocalSize().X, SelectedComponentGeometry.GetLocalSize().Y * 0.5f);
 				case EDMXPixelMappingTransformDirection::BottomCenter:
-					return FVector2f(SelectedComponentGeometry.GetLocalSize().X * 0.5f, SelectedComponentGeometry.GetLocalSize().Y);
+					return FVector2D(SelectedComponentGeometry.GetLocalSize().X * 0.5f, SelectedComponentGeometry.GetLocalSize().Y);
 				case EDMXPixelMappingTransformDirection::BottomRight:
-					return UE::Slate::CastToVector2f(SelectedComponentGeometry.GetLocalSize());
+					return SelectedComponentGeometry.GetLocalSize();
 				default:
 					checkNoEntry(); // Unhandled enum value
 				}
-				return FVector2f::ZeroVector;
+				return FVector2D::ZeroVector;
 			}();
 
-			const auto RotateVectorAroundPivot =
-				[](FVector2f Point, FVector2f Pivot, double AngleDegrees)
-				{
-					double Sin;
-					double Cos;
-					FMath::SinCos(&Sin, &Cos, FMath::DegreesToRadians(AngleDegrees));
-
-					const FVector2f RelativePositionOld = Point - Pivot;
-					const FVector2f RelativePositionNew = FVector2f(
-						Cos * RelativePositionOld.X - Sin * RelativePositionOld.Y,
-						Sin * RelativePositionOld.X + Cos * RelativePositionOld.Y);
-
-					return RelativePositionNew + Pivot;
-				};
-			const FVector2f LocalPositionRotated = RotateVectorAroundPivot(LocalPosition, SelectedComponentGeometry.GetLocalSize() / 2.f, OutputComponent->GetRotation());
 
-			const FVector2f SelectedWidgetScale = SelectedComponentGeometry.GetAccumulatedRenderTransform().GetMatrix().GetScale().GetVector();
+			FVector2D SelectedWidgetScale = FVector2D(SelectedComponentGeometry.GetAccumulatedRenderTransform().GetMatrix().GetScale().GetVector());
 
-			const FVector2f ApplicationScaledOffset = UE::Slate::CastToVector2f(Handle->GetOffset() * GetDesignerGeometry().Scale);
+			FVector2D ApplicationScaledOffset = Handle->GetOffset() * GetDesignerGeometry().Scale;
 
-			const FVector2f LocalOffsetFull = ApplicationScaledOffset / SelectedWidgetScale;
-			const FVector2f PositionFullOffset = GetDesignerGeometry().AbsoluteToLocal(SelectedComponentGeometry.LocalToAbsolute(LocalPositionRotated + LocalOffsetFull));
-			const FVector2f LocalOffsetHalf = (ApplicationScaledOffset / 2.0f) / SelectedWidgetScale;
-			const FVector2f PositionHalfOffset = GetDesignerGeometry().AbsoluteToLocal(SelectedComponentGeometry.LocalToAbsolute(LocalPositionRotated + LocalOffsetHalf));
+			FVector2D LocalOffsetFull = ApplicationScaledOffset / SelectedWidgetScale;
+			FVector2D PositionFullOffset = GetDesignerGeometry().AbsoluteToLocal(SelectedComponentGeometry.LocalToAbsolute(LocalPosition + LocalOffsetFull));
+			FVector2D LocalOffsetHalf = (ApplicationScaledOffset / 2.0f) / SelectedWidgetScale;
+			FVector2D PositionHalfOffset = GetDesignerGeometry().AbsoluteToLocal(SelectedComponentGeometry.LocalToAbsolute(LocalPosition + LocalOffsetHalf));
 
-			FVector2f PivotCorrection = PositionHalfOffset - (PositionFullOffset + FVector2f(5.0f, 5.0f));
+			FVector2D PivotCorrection = PositionHalfOffset - (PositionFullOffset + FVector2D(5.0f, 5.0f));
 
-			const FVector2f FinalPosition = PositionFullOffset + PivotCorrection;
+			FVector2D FinalPosition = PositionFullOffset + PivotCorrection;
 
 			return FVector2D(FinalPosition);
 		}
@@ -1030,6 +1015,7 @@ bool SDMXPixelMappingDesignerView::GetComponentGeometry(UDMXPixelMappingBaseComp
 {
 	if (UDMXPixelMappingOutputComponent* OutputComponent = Cast<UDMXPixelMappingOutputComponent>(InBaseComponent))
 	{
+		using namespace UE::DMX;
 		const TSharedRef<IDMXPixelMappingOutputComponentWidgetInterface>* ComponentViewPtr = Algo::FindByPredicate(OutputComponentWidgets, [InBaseComponent](const TSharedRef<IDMXPixelMappingOutputComponentWidgetInterface>& ComponentWidget)
 			{
 				return ComponentWidget->Equals(InBaseComponent);
diff --git a/Engine/Plugins/VirtualProduction/DMX/DMXPixelMapping/Source/DMXPixelMappingEditor/Private/Views/SDMXPixelMappingDesignerView.h b/Engine/Plugins/VirtualProduction/DMX/DMXPixelMapping/Source/DMXPixelMappingEditor/Private/Views/SDMXPixelMappingDesignerView.h
--- a/Engine/Plugins/VirtualProduction/DMX/DMXPixelMapping/Source/DMXPixelMappingEditor/Private/Views/SDMXPixelMappingDesignerView.h
+++ b/Engine/Plugins/VirtualProduction/DMX/DMXPixelMapping/Source/DMXPixelMappingEditor/Private/Views/SDMXPixelMappingDesignerView.h
@@ -8,7 +8,6 @@
 class FDMXPixelMappingComponentReference;
 class FDMXPixelMappingDragDropOp;
 class FDMXPixelMappingToolkit;
-class IDMXPixelMappingOutputComponentWidgetInterface;
 class SBorder;
 class SBox;
 class SCanvas;
@@ -22,6 +21,10 @@ class UDMXPixelMapping;
 class UDMXPixelMappingBaseComponent;
 class UDMXPixelMappingOutputComponent;
 class UDMXPixelMappingRendererComponent;
+namespace UE::DMX
+{
+	class IDMXPixelMappingOutputComponentWidgetInterface;
+}
 
 
 class SDMXPixelMappingDesignerView
@@ -205,7 +208,7 @@ private:
 	FVector2D DragAnchor;
 
 	/** Output component views currently displayed */
-	TArray<TSharedRef<IDMXPixelMappingOutputComponentWidgetInterface>> OutputComponentWidgets;
+	TArray<TSharedRef<UE::DMX::IDMXPixelMappingOutputComponentWidgetInterface>> OutputComponentWidgets;
 
 	TArray<TSharedPtr<SDMXPixelMappingTransformHandle>> TransformHandles;
 };
diff --git a/Engine/Plugins/VirtualProduction/DMX/DMXPixelMapping/Source/DMXPixelMappingEditor/Private/Views/SDMXPixelMappingPreviewView.cpp b/Engine/Plugins/VirtualProduction/DMX/DMXPixelMapping/Source/DMXPixelMappingEditor/Private/Views/SDMXPixelMappingPreviewView.cpp
--- a/Engine/Plugins/VirtualProduction/DMX/DMXPixelMapping/Source/DMXPixelMappingEditor/Private/Views/SDMXPixelMappingPreviewView.cpp
+++ b/Engine/Plugins/VirtualProduction/DMX/DMXPixelMapping/Source/DMXPixelMappingEditor/Private/Views/SDMXPixelMappingPreviewView.cpp
@@ -294,7 +294,7 @@ TSharedRef<SWidget> SDMXPixelMappingPreviewView::CreateOverlayUI()
 				.ContentPadding(FAppStyle::Get().GetMargin("ViewportMenu.SToolBarButtonBlock.Button.Padding"))
 				[
 					SNew(SImage)
-					.Image(FDMXPixelMappingEditorStyle::Get().GetBrush("Icons.ZoomToFit"))
+					.Image(FAppStyle::GetBrush("UMGEditor.ZoomToFit"))
 				]
 			]
 		];
diff --git a/Engine/Plugins/VirtualProduction/DMX/DMXPixelMapping/Source/DMXPixelMappingEditor/Private/Widgets/SDMXPixelMappingOutputComponent.cpp b/Engine/Plugins/VirtualProduction/DMX/DMXPixelMapping/Source/DMXPixelMappingEditor/Private/Widgets/SDMXPixelMappingOutputComponent.cpp
--- a/Engine/Plugins/VirtualProduction/DMX/DMXPixelMapping/Source/DMXPixelMappingEditor/Private/Widgets/SDMXPixelMappingOutputComponent.cpp
+++ b/Engine/Plugins/VirtualProduction/DMX/DMXPixelMapping/Source/DMXPixelMappingEditor/Private/Widgets/SDMXPixelMappingOutputComponent.cpp
@@ -6,6 +6,7 @@
 #include "DMXPixelMappingTypes.h"
 #include "DMXPixelMappingUtils.h"
 #include "DMXRuntimeUtils.h"
+#include "Fonts/FontMeasure.h"
 #include "Settings/DMXPixelMappingEditorSettings.h"
 #include "Toolkits/DMXPixelMappingToolkit.h"
 #include "ViewModels/DMXPixelMappingOutputComponentModel.h"
@@ -13,691 +14,514 @@
 #include "Widgets/Layout/SBox.h"
 #include "Widgets/Layout/SScaleBox.h"
 #include "Widgets/Layout/SUniformGridPanel.h"
+#include "Widgets/SDMXPixelMappingOutputComponentText.h"
 #include "Widgets/Text/STextBlock.h"
 
 
 #define LOCTEXT_NAMESPACE "SDMXPixelMappingOutputComponent"
 
-IDMXPixelMappingOutputComponentWidgetInterface::~IDMXPixelMappingOutputComponentWidgetInterface()
+namespace UE::DMX
 {
-	RemoveFromCanvas();
-}
-
-void IDMXPixelMappingOutputComponentWidgetInterface::AddToCanvas(const TSharedRef<SConstraintCanvas>& InCanvas)
-{
-	RemoveFromCanvas();
-
-	ParentCanvas = InCanvas;
-
-	ParentCanvas->AddSlot()
-		.ZOrder(0)
-		.AutoSize(true)
-		.Alignment(FVector2D::ZeroVector)
-		.Offset_Lambda([this]()
-			{
-				const FVector2D Position = GetPosition();
-				return FMargin(Position.X, Position.Y);
-			})
-		.Expose(Slot)
-		[
-			AsWidget()
-		];
-}
-
-void IDMXPixelMappingOutputComponentWidgetInterface::RemoveFromCanvas()
-{
-	if (ParentCanvas.IsValid() && Slot)
+	IDMXPixelMappingOutputComponentWidgetInterface::~IDMXPixelMappingOutputComponentWidgetInterface()
 	{
-		ParentCanvas->RemoveSlot(Slot->GetWidget());
+		RemoveFromCanvas();
 	}
 
-	Slot = nullptr;
-	ParentCanvas.Reset();
-}
-
-
-void SDMXPixelMappingOutputComponent::Construct(const FArguments& InArgs, const TSharedRef<FDMXPixelMappingToolkit>& InToolkit, TWeakObjectPtr<UDMXPixelMappingOutputComponent> OutputComponent)
-{
-	if (!OutputComponent.IsValid())
+	void IDMXPixelMappingOutputComponentWidgetInterface::AddToCanvas(const TSharedRef<SConstraintCanvas>& InCanvas)
 	{
-		return;
-	}
-
-	WeakToolkit = InToolkit;
-	Model = MakeShared<FDMXPixelMappingOutputComponentModel>(InToolkit, OutputComponent);
-
-	ChildSlot
-	[
-		SAssignNew(ComponentBox, SBox)
-		.WidthOverride_Lambda([this]()
-			{
-				return Model->GetSize().X;
-			})
-		.HeightOverride_Lambda([this]()
-			{
-				return Model->GetSize().Y;
-			})
-	];
+		RemoveFromCanvas();
 
-	ForceRefresh();
+		ParentCanvas = InCanvas;
 
-	UDMXPixelMappingEditorSettings::OnEditorSettingsChanged.AddSP(this, &SDMXPixelMappingOutputComponent::RequestRefresh);
-}
-
-bool SDMXPixelMappingOutputComponent::Equals(UDMXPixelMappingBaseComponent* Component) const
-{
-	if (Component)
-	{
-		return Model->Equals(Component);
+		ParentCanvas->AddSlot()
+			.ZOrder(0)
+			.AutoSize(true)
+			.Alignment(FVector2D::ZeroVector)
+			.Offset_Lambda([this]()
+				{
+					const FVector2D Position = GetPosition();
+					return FMargin(Position.X, Position.Y);
+				})
+			.Expose(Slot)
+			[
+				AsWidget()
+			];
 	}
-	return false;
-}
-
-FVector2D SDMXPixelMappingOutputComponent::GetPosition() const
-{
-	return Model->GetPosition();
-}
 
-void SDMXPixelMappingOutputComponent::RequestRefresh()
-{
-	if (!RefreshTimerHandle.IsValid())
+	void IDMXPixelMappingOutputComponentWidgetInterface::RemoveFromCanvas()
 	{
-		RefreshTimerHandle = GEditor->GetTimerManager()->SetTimerForNextTick(FTimerDelegate::CreateSP(this, &SDMXPixelMappingOutputComponent::ForceRefresh));
-	}
-}
+		if (ParentCanvas.IsValid() && Slot)
+		{
+			ParentCanvas->RemoveSlot(Slot->GetWidget());
+		}
 
-void SDMXPixelMappingOutputComponent::ForceRefresh()
-{
-	RefreshTimerHandle.Invalidate();
-	if (Model->ShouldDraw())
-	{
-		ComponentBox->SetContent(CreateContent());
-	}
-	else
-	{
-		ComponentBox->SetContent(SNullWidget::NullWidget);
+		Slot = nullptr;
+		ParentCanvas.Reset();
 	}
-}
-
-TSharedRef<SWidget> SDMXPixelMappingOutputComponent::CreateContent()
-{
-	const TSharedRef<SBorder> Content =
-		SNew(SBorder)
-		.RenderTransform(this, &SDMXPixelMappingOutputComponent::ComputeRenderTransform)
-		.RenderTransformPivot(FVector2D(0.5, 0.5))
-		.BorderImage(FDMXPixelMappingEditorStyle::Get().GetBrush("DMXPixelMappingEditor.ComponentBorder"))
-		.BorderBackgroundColor_Lambda([this]()
-			{
-				return Model->GetColor();
-			})
-		[
-			SNew(SVerticalBox)
 
-			+ SVerticalBox::Slot()
-			.HAlign(HAlign_Fill)
-			.VAlign(VAlign_Bottom)
-			.AutoHeight()
-			.Padding(FMargin(0.f, -22.f, 0.f, 0.f))
-			[
-				SNew(SBox)
-				.HeightOverride(22.f)
-				.HAlign(HAlign_Left)
-				.VAlign(VAlign_Bottom)
-				[
-					SAssignNew(AboveContentBorder, SBorder)
-					.BorderImage(FAppStyle::GetBrush("NoBorder"))
-				]
-			]
-
-			+ SVerticalBox::Slot()
-			.HAlign(HAlign_Fill)
-			.VAlign(VAlign_Fill)
-			.FillHeight(1.f / 3.f)
-			.Padding(FMargin(-2.f, -2.f, -2.f, -2.f))
-			[
-				SAssignNew(TopContentBorder, SBorder)
-				.BorderImage(FAppStyle::GetBrush("NoBorder"))
-			]
+	void SDMXPixelMappingOutputComponent::Construct(const FArguments& InArgs, const TSharedRef<FDMXPixelMappingToolkit>& InToolkit, TWeakObjectPtr<UDMXPixelMappingOutputComponent> InOutputComponent)
+	{
+		WeakToolkit = InToolkit;
+		Model = MakeShared<FDMXPixelMappingOutputComponentModel>(InToolkit, InOutputComponent);
+		if (!InOutputComponent.IsValid())
+		{
+			return;
+		}
 
-			+ SVerticalBox::Slot()
-			.HAlign(HAlign_Fill)
-			.VAlign(VAlign_Fill)
-			.FillHeight(1.f / 3.f)
-			.Padding(FMargin(-2.f, -2.f, -2.f, -2.f))
-			[
-				SAssignNew(MiddleContentBorder, SBorder)
-				.BorderImage(FAppStyle::GetBrush("NoBorder"))
-			]
+		SetRenderTransform(TAttribute<TOptional<FSlateRenderTransform>>::CreateSP(this, &SDMXPixelMappingOutputComponent::GetRenderTransform));
+		SetRenderTransformPivot(FVector2D(0.5, 0.5));
 
-			+ SVerticalBox::Slot()
-			.HAlign(HAlign_Fill)
-			.VAlign(VAlign_Fill)
-			.FillHeight(1.f / 3.f)
-			.Padding(FMargin(-2.f, -2.f, -2.f, -2.f))
+		// Define the bounding box, rest is painted OnPaint.
+		ChildSlot
+		[
+			SNew(SBox)
+			.WidthOverride_Lambda([this]()
+				{
+					return Model->GetSize().X;
+				})
+			.HeightOverride_Lambda([this]()
+				{
+					return Model->GetSize().Y;
+				})
 			[
-				SAssignNew(BottomContentBorder, SBorder)
-				.BorderImage(FAppStyle::GetBrush("NoBorder"))
+				SNew(SDMXPixelMappingOutputComponentText, InToolkit, InOutputComponent)
 			]
 		];
-
-	if (Model->ShouldDrawName())
-	{
-		if (Model->ShouldDrawNameAbove())
-		{
-			CreateComponentNameChildSlotAbove();
-		}
-		else
-		{
-			CreateComponentNameChildSlotInside();
-		}
 	}
 
-	if (Model->ShouldDrawCellID())
+	bool SDMXPixelMappingOutputComponent::Equals(UDMXPixelMappingBaseComponent* Component) const
 	{
-		CreateCellIDChildSlot();
+		return Model->Equals(Component);
 	}
 
-	if (Model->ShouldDrawPatchInfo())
+	FVector2D SDMXPixelMappingOutputComponent::GetPosition() const
 	{
-		CreatePatchInfoChildSlot();
+		return Model->GetPosition();
 	}
 
-	return Content;
-}
+	int32 SDMXPixelMappingOutputComponent::OnPaint(
+		const FPaintArgs& Args,
+		const FGeometry& AllottedGeometry,
+		const FSlateRect& MyCullingRect,
+		FSlateWindowElementList& OutDrawElements,
+		int32 LayerId,
+		const FWidgetStyle& InWidgetStyle,
+		bool bParentEnabled) const
+	{
+		if (!Model->ShouldDraw())
+		{
+			return LayerId;
+		}
+		LayerId = SCompoundWidget::OnPaint(Args, AllottedGeometry, MyCullingRect, OutDrawElements, LayerId, InWidgetStyle, bParentEnabled);
 
-void SDMXPixelMappingOutputComponent::CreateComponentNameChildSlotAbove()
-{
-	AboveContentBorder->SetContent(
-		SNew(SBox)
-		.HeightOverride(20.f)
-		.HAlign(HAlign_Left)
-		.VAlign(VAlign_Bottom)
-		[
-			SNew(SScaleBox)
-			.Stretch(EStretch::ScaleToFit)
-			.StretchDirection(EStretchDirection::DownOnly)
-			[
-				SNew(STextBlock)
-				.Font(FCoreStyle::GetDefaultFontStyle("Regular", 18))
-				.Text_Lambda([this]()
-					{
-						return Model->GetName();
-					})
-			]	
-		]
+		const FSlateBrush* BorderBrush = FDMXPixelMappingEditorStyle::Get().GetBrush("DMXPixelMappingEditor.ComponentBorder");
+		const FLinearColor Color = Model->GetColor();
 
-	);
-}
+		const FVector2D Size = Model->GetSize();
+		TArray<FVector2D> LinePoints
+		{
+			FVector2D::ZeroVector,
+			FVector2D(0.0, Size.Y),
+			Size,
+			FVector2D(Size.X, 0.0),
+			FVector2D::ZeroVector
+		};
+
+		constexpr float LineThickness = 2.f;
+		constexpr bool bAntialias = true;
+		FSlateDrawElement::MakeLines(
+			OutDrawElements,
+			LayerId,
+			AllottedGeometry.ToPaintGeometry(),
+			LinePoints,
+			ESlateDrawEffect::None,
+			Color,
+			bAntialias,
+			LineThickness);
+
+		return LayerId + 1;
+	}
 
-void SDMXPixelMappingOutputComponent::CreateComponentNameChildSlotInside()
-{
-	TopContentBorder->SetContent(
-		SNew(SBox)
-		.HeightOverride(14.f)
-		.HAlign(HAlign_Left)
-		.VAlign(VAlign_Top)
-		[
-			SNew(SScaleBox)
-			.Stretch(EStretch::ScaleToFit)
-			.StretchDirection(EStretchDirection::DownOnly)
-			.HAlign(HAlign_Left)
-			.VAlign(VAlign_Top)
-			[
-				SNew(STextBlock)
-				.Font(FCoreStyle::GetDefaultFontStyle("Regular", 12))
-				.Text_Lambda([this]()
-					{
-						return Model->GetName();
-					})
-			]
-		]
-	);
-}
+	TOptional<FSlateRenderTransform> SDMXPixelMappingOutputComponent::GetRenderTransform() const
+	{	
+		const FQuat2D Quaternion = Model->GetQuaternion();
+	
+		return FSlateRenderTransform(Quaternion, FVector2D::ZeroVector);
+	}
 
-void SDMXPixelMappingOutputComponent::CreateCellIDChildSlot()
-{	
-	MiddleContentBorder->SetContent(
-		SNew(SBox)
-		.HeightOverride(14.f)
-		.HAlign(HAlign_Center)
-		.VAlign(VAlign_Center)
-		[
-			SNew(SScaleBox)
-			.Stretch(EStretch::ScaleToFit)
-			.StretchDirection(EStretchDirection::DownOnly)
-			[
-				SNew(STextBlock)
-				.Font(FCoreStyle::GetDefaultFontStyle("Regular", 12))
-				.Text_Lambda([this]()
-					{
-						return Model->GetCellIDText();
-					})
-			]
-		]
-	);
-}
+	class SDMXPixelMappingScreenComponentCell
+		: public SCompoundWidget
+	{
+	public:
+		SLATE_BEGIN_ARGS(SDMXPixelMappingScreenComponentCell)
+			: _ChannelOffset(0)
+			, _StartingChannel(1)
+			, _Universe(1)
+		{}
+			/** The offset of this cell from the starting address */
+			SLATE_ARGUMENT(int32, ChannelOffset)
 
-void SDMXPixelMappingOutputComponent::CreatePatchInfoChildSlot()
-{
-	BottomContentBorder->SetContent(
-		SNew(SBox)
-		.HeightOverride(8.f)
-		.HAlign(HAlign_Left)
-		.VAlign(VAlign_Bottom)
-		[
-			SNew(SScaleBox)
-			.Stretch(EStretch::ScaleToFit)
-			.StretchDirection(EStretchDirection::DownOnly)
-			[
-				SNew(SVerticalBox)
+			/** The starting address of the screen component */
+			SLATE_ATTRIBUTE(int32, StartingChannel)
 
-				+ SVerticalBox::Slot()
-				.HAlign(HAlign_Left)
-				.VAlign(VAlign_Bottom)
-				[
-					SNew(STextBlock)
-					.ColorAndOpacity(FLinearColor::Green)
-					.Font(FCoreStyle::GetDefaultFontStyle("Regular", 6))
-					.Text_Lambda([this]()
-						{
-							return Model->GetAddressesText();
-						})
-				]
+			/** The universe of the screen component */
+			SLATE_ATTRIBUTE(int32, Universe)
 
-				+ SVerticalBox::Slot()
-				.HAlign(HAlign_Left)
-				.VAlign(VAlign_Bottom)
-				[
-					SNew(STextBlock)
-					.ColorAndOpacity(FLinearColor::Green)
-					.Font(FCoreStyle::GetDefaultFontStyle("Regular", 6))
-					.Text_Lambda([this]()
-						{
-							return Model->GetFixtureIDText();
-						})
-				]
-			]
-		]
-	);
-}
+		SLATE_END_ARGS()
 
-TOptional<FSlateRenderTransform> SDMXPixelMappingOutputComponent::ComputeRenderTransform() const
-{
-	return FSlateRenderTransform(Model->GetQuaternion(), FVector2D(0.0, 0.0));
-}
+		/** Constructs this widget */
+		void Construct(const FArguments& InArgs, const TSharedRef<FDMXPixelMappingScreenComponentModel>& InModel)
+		{
+			Model = InModel;
 
-class SDMXPixelMappingScreenComponentCell
-	: public SCompoundWidget
-{
-public:
-	SLATE_BEGIN_ARGS(SDMXPixelMappingScreenComponentCell)
-		: _ChannelOffset(0)
-		, _StartingChannel(1)
-		, _Universe(1)
-	{}
-		/** The offset of this cell from the starting address */
-		SLATE_ARGUMENT(int32, ChannelOffset)
+			ChannelOffset = InArgs._ChannelOffset;
+			Universe = InArgs._Universe;
+			StartingChannel = InArgs._StartingChannel;
 
-		/** The starting address of the screen component */
-		SLATE_ATTRIBUTE(int32, StartingChannel)
+			ChildSlot
+				[
+					SAssignNew(ContentBorder, SBorder)
+					.BorderImage(FDMXPixelMappingEditorStyle::Get().GetBrush("DMXPixelMappingEditor.ComponentBorder"))
+					.HAlign(HAlign_Fill)
+					.VAlign(VAlign_Fill)
+				];
 
-		/** The universe of the screen component */
-		SLATE_ATTRIBUTE(int32, Universe)
+			RebuildAddressesContent();
 
-	SLATE_END_ARGS()
+			UDMXPixelMappingEditorSettings::OnEditorSettingsChanged.AddSP(this, &SDMXPixelMappingScreenComponentCell::RebuildAddressesContent);
+		}
 
-	/** Constructs this widget */
-	void Construct(const FArguments& InArgs, const TSharedRef<FDMXPixelMappingScreenComponentModel>& InModel)
-	{
-		Model = InModel;
+	private:
+		/** Rebuilds addresses content */
+		void RebuildAddressesContent()
+		{
+			const FDMXPixelMappingDesignerSettings& DesignerSettings = GetDefault<UDMXPixelMappingEditorSettings>()->DesignerSettings;
+			if (!DesignerSettings.bShowPatchInfo)
+			{
+				return;
+			}
 
-		ChannelOffset = InArgs._ChannelOffset;
-		Universe = InArgs._Universe;
-		StartingChannel = InArgs._StartingChannel;
+			if (!ensureMsgf(StartingChannel.IsSet() && Universe.IsSet(), TEXT("Trying to display universe and address of pixel mapping screen component, but universe and/or address are not set.")))
+			{
+				return;
+			}
 
-		ChildSlot
-			[
-				SAssignNew(ContentBorder, SBorder)
-				.BorderImage(FDMXPixelMappingEditorStyle::Get().GetBrush("DMXPixelMappingEditor.ComponentBorder"))
+			ContentBorder->SetContent
+			(
+				SNew(SBorder)
+				.BorderImage(FAppStyle::GetBrush("NoBorder"))
 				.HAlign(HAlign_Fill)
 				.VAlign(VAlign_Fill)
-			];
-
-		RebuildAddressesContent();
-
-		UDMXPixelMappingEditorSettings::OnEditorSettingsChanged.AddSP(this, &SDMXPixelMappingScreenComponentCell::RebuildAddressesContent);
-	}
-
-private:
-	/** Rebuilds addresses content */
-	void RebuildAddressesContent()
-	{
-		const FDMXPixelMappingDesignerSettings& DesignerSettings = GetDefault<UDMXPixelMappingEditorSettings>()->DesignerSettings;
-		if (!DesignerSettings.bShowPatchInfo)
-		{
-			return;
-		}
-
-		if (!ensureMsgf(StartingChannel.IsSet() && Universe.IsSet(), TEXT("Trying to display universe and address of pixel mapping screen component, but universe and/or address are not set.")))
-		{
-			return;
-		}
-
-		ContentBorder->SetContent
-		(
-			SNew(SBorder)
-			.BorderImage(FAppStyle::GetBrush("NoBorder"))
-			.HAlign(HAlign_Fill)
-			.VAlign(VAlign_Fill)
-			.Padding(6.f)
-			[
-				SNew(SScaleBox)
-				.Stretch(EStretch::ScaleToFit)
-				.StretchDirection(EStretchDirection::DownOnly)
-				.VAlign(VAlign_Center)
-				.HAlign(HAlign_Center)
+				.Padding(6.f)
 				[
-					SNew(STextBlock)
-					.Font(FCoreStyle::GetDefaultFontStyle("Regular", 12))
-					.Text_Lambda([this]()
-						{
-							const int32 UniverseWithOffset = Universe.Get() + ChannelOffset / DMX_MAX_ADDRESS;
-							const int32 ChannelWithOffset = [this, UniverseWithOffset]()
+					SNew(SScaleBox)
+					.Stretch(EStretch::ScaleToFit)
+					.StretchDirection(EStretchDirection::DownOnly)
+					.VAlign(VAlign_Center)
+					.HAlign(HAlign_Center)
+					[
+						SNew(STextBlock)
+						.Font(FCoreStyle::GetDefaultFontStyle("Regular", 12))
+						.Text_Lambda([this]()
 							{
-								if (UniverseWithOffset > 0)
+								const int32 UniverseWithOffset = Universe.Get() + ChannelOffset / DMX_MAX_ADDRESS;
+								const int32 ChannelWithOffset = [this, UniverseWithOffset]()
+								{
+									if (UniverseWithOffset > 0)
+									{
+										return (StartingChannel.Get() + ChannelOffset) % DMX_MAX_ADDRESS;
+									}
+									return StartingChannel.Get() + ChannelOffset;
+								}();
+
+								if (Model->ComponentWantsToShowChannel() && Model->ComponentWantsToShowUniverse())
 								{
-									return (StartingChannel.Get() + ChannelOffset) % DMX_MAX_ADDRESS;
+									FString AddressString = FString::FromInt(UniverseWithOffset) + TEXT(".") + FString::FromInt(ChannelWithOffset);
+									return FText::FromString(AddressString);
 								}
-								return StartingChannel.Get() + ChannelOffset;
-							}();
-
-							if (Model->ComponentWantsToShowChannel() && Model->ComponentWantsToShowUniverse())
-							{
-								FString AddressString = FString::FromInt(UniverseWithOffset) + TEXT(".") + FString::FromInt(ChannelWithOffset);
-								return FText::FromString(AddressString);
-							}
-							else if (Model->ComponentWantsToShowUniverse())
-							{
-								FString UniverseString = FString::FromInt(UniverseWithOffset);
-								return FText::FromString(UniverseString);
-							}
-							else if (Model->ComponentWantsToShowChannel())
-							{
-								FString ChannelString = FString::FromInt(ChannelWithOffset);
-								return FText::FromString(ChannelString);
-							}
-							return FText::GetEmpty();
-						})
+								else if (Model->ComponentWantsToShowUniverse())
+								{
+									FString UniverseString = FString::FromInt(UniverseWithOffset);
+									return FText::FromString(UniverseString);
+								}
+								else if (Model->ComponentWantsToShowChannel())
+								{
+									FString ChannelString = FString::FromInt(ChannelWithOffset);
+									return FText::FromString(ChannelString);
+								}
+								return FText::GetEmpty();
+							})
+					]
 				]
-			]
-		);
-	}
-
-	/** Border to hold contents of the widget */
-	TSharedPtr<SBorder> ContentBorder;
+			);
+		}
 
-	/** The model used by this widget */
-	TSharedPtr<FDMXPixelMappingScreenComponentModel> Model;
+		/** Border to hold contents of the widget */
+		TSharedPtr<SBorder> ContentBorder;
 
-	// Slate args
-	int32 ChannelOffset = 0;
-	TAttribute<int32> Universe;
-	TAttribute<int32> StartingChannel;
-};
+		/** The model used by this widget */
+		TSharedPtr<FDMXPixelMappingScreenComponentModel> Model;
 
-/** Widget to draw the inners of a screen component */
-class SDMXPixelMappingScreenComponentGrid
-	: public SCompoundWidget
-{
-public:
-	SLATE_BEGIN_ARGS(SDMXPixelMappingScreenComponentGrid)
-	{}
-	SLATE_END_ARGS()
+		// Slate args
+		int32 ChannelOffset = 0;
+		TAttribute<int32> Universe;
+		TAttribute<int32> StartingChannel;
+	};
 
-	/** Constructs this widget */
-	void Construct(const FArguments& InArgs, const TSharedRef<FDMXPixelMappingScreenComponentModel>& InModel)
+	/** Widget to draw the inners of a screen component */
+	class SDMXPixelMappingScreenComponentGrid
+		: public SCompoundWidget
 	{
-		Model = InModel;
-
-		ChildSlot
-			.VAlign(VAlign_Fill)
-			.HAlign(HAlign_Fill)
-			[
-				SAssignNew(GridPanel, SUniformGridPanel)
-			];
+	public:
+		SLATE_BEGIN_ARGS(SDMXPixelMappingScreenComponentGrid)
+		{}
+		SLATE_END_ARGS()
 
-		CachedCellFormat = Model->GetCellFormat();
-		CachedDistribution = Model->GetDistribution();
+		/** Constructs this widget */
+		void Construct(const FArguments& InArgs, const TSharedRef<FDMXPixelMappingScreenComponentModel>& InModel)
+		{
+			Model = InModel;
 
-		RebuildGridContent();
-	}
+			ChildSlot
+				.VAlign(VAlign_Fill)
+				.HAlign(HAlign_Fill)
+				[
+					SAssignNew(GridPanel, SUniformGridPanel)
+				];
 
-protected:
-	//~ Begin SWidget Interface
-	virtual void Tick(const FGeometry& AllottedGeometry, const double InCurrentTime, const float InDeltaTime) override
-	{
-		if ((GridPanel->GetChildren()->Num() != Model->GetNumColumns() * Model->GetNumRows()) ||
-			CachedDistribution != Model->GetDistribution() ||
-			CachedCellFormat != Model->GetCellFormat())
-		{
-			// Handle cell format changes
-			CachedDistribution = Model->GetDistribution();
 			CachedCellFormat = Model->GetCellFormat();
+			CachedDistribution = Model->GetDistribution();
 
 			RebuildGridContent();
 		}
-	}
-	//~ End SWidget Interface
-
-private:
-	/** Rebuilds the grid content */
-	void RebuildGridContent()
-	{
-		GridPanel->ClearChildren();
 
-		TArray<int32> CellIndices;
-		const int32 TotalCells = Model->GetNumColumns() * Model->GetNumRows();
-		for (int32 CellID = 0; CellID < TotalCells; ++CellID)
+	protected:
+		//~ Begin SWidget Interface
+		virtual void Tick(const FGeometry& AllottedGeometry, const double InCurrentTime, const float InDeltaTime) override
 		{
-			CellIndices.Add(CellID);
+			if ((GridPanel->GetChildren()->Num() != Model->GetNumColumns() * Model->GetNumRows()) ||
+				CachedDistribution != Model->GetDistribution() ||
+				CachedCellFormat != Model->GetCellFormat())
+			{
+				// Handle cell format changes
+				CachedDistribution = Model->GetDistribution();
+				CachedCellFormat = Model->GetCellFormat();
+
+				RebuildGridContent();
+			}
 		}
+		//~ End SWidget Interface
 
-		TArray<int32> DistributionSortedCellIndices;
-		FDMXRuntimeUtils::PixelMappingDistributionSort(Model->GetDistribution(), Model->GetNumColumns(), Model->GetNumRows(), CellIndices, DistributionSortedCellIndices);
+	private:
+		/** Rebuilds the grid content */
+		void RebuildGridContent()
+		{
+			GridPanel->ClearChildren();
 
-		const uint32 NumChannelsPerCell = FDMXPixelMappingUtils::GetNumChannelsPerCell(Model->GetCellFormat());
+			TArray<int32> CellIndices;
+			const int32 TotalCells = Model->GetNumColumns() * Model->GetNumRows();
+			for (int32 CellID = 0; CellID < TotalCells; ++CellID)
+			{
+				CellIndices.Add(CellID);
+			}
 
-		uint32 XYIndex = 0;
-		for (int32 YIndex = 0; YIndex < Model->GetNumRows(); ++YIndex)
-		{
-			for (int32 XIndex = 0; XIndex < Model->GetNumColumns(); ++XIndex)
+			TArray<int32> DistributionSortedCellIndices;
+			FDMXRuntimeUtils::PixelMappingDistributionSort(Model->GetDistribution(), Model->GetNumColumns(), Model->GetNumRows(), CellIndices, DistributionSortedCellIndices);
+
+			const uint32 NumChannelsPerCell = FDMXPixelMappingUtils::GetNumChannelsPerCell(Model->GetCellFormat());
+
+			uint32 XYIndex = 0;
+			for (int32 YIndex = 0; YIndex < Model->GetNumRows(); ++YIndex)
 			{
-				if (!ensureMsgf(DistributionSortedCellIndices.IsValidIndex(XYIndex), TEXT("Pixel mapping cell index out of range.")))
+				for (int32 XIndex = 0; XIndex < Model->GetNumColumns(); ++XIndex)
 				{
-					break;
-				}
-
-				const int32 ChannelOffset = DistributionSortedCellIndices[XYIndex] * NumChannelsPerCell;
-				GridPanel->AddSlot(XIndex, YIndex)
-					.HAlign(HAlign_Fill)
-					.VAlign(VAlign_Fill)
-					[
-						SNew(SDMXPixelMappingScreenComponentCell, Model.ToSharedRef())
-						.ChannelOffset(ChannelOffset)
-						.StartingChannel_Lambda([this]()
-							{
-								return Model->GetStartingChannel();
-							})
-						.Universe_Lambda([this]()
-							{
-								return Model->GetUniverse();
-							})
-					];
+					if (!ensureMsgf(DistributionSortedCellIndices.IsValidIndex(XYIndex), TEXT("Pixel mapping cell index out of range.")))
+					{
+						break;
+					}
+
+					const int32 ChannelOffset = DistributionSortedCellIndices[XYIndex] * NumChannelsPerCell;
+					GridPanel->AddSlot(XIndex, YIndex)
+						.HAlign(HAlign_Fill)
+						.VAlign(VAlign_Fill)
+						[
+							SNew(SDMXPixelMappingScreenComponentCell, Model.ToSharedRef())
+							.ChannelOffset(ChannelOffset)
+							.StartingChannel_Lambda([this]()
+								{
+									return Model->GetStartingChannel();
+								})
+							.Universe_Lambda([this]()
+								{
+									return Model->GetUniverse();
+								})
+						];
 
-				XYIndex++;
+					XYIndex++;
+				}
 			}
 		}
-	}
 
-	/** Cell format, cached to detect changes */
-	EDMXCellFormat CachedCellFormat = EDMXCellFormat::PF_RGB;
+		/** Cell format, cached to detect changes */
+		EDMXCellFormat CachedCellFormat = EDMXCellFormat::PF_RGB;
 
-	/** Pixelmapping distribution, cached to detect changes */
-	EDMXPixelMappingDistribution CachedDistribution = EDMXPixelMappingDistribution::TopLeftToRight;
+		/** Pixelmapping distribution, cached to detect changes */
+		EDMXPixelMappingDistribution CachedDistribution = EDMXPixelMappingDistribution::TopLeftToRight;
 
-	/** The actual grid panel */
-	TSharedPtr<SUniformGridPanel> GridPanel;
+		/** The actual grid panel */
+		TSharedPtr<SUniformGridPanel> GridPanel;
 
-	/** The model used by this widget */
-	TSharedPtr<FDMXPixelMappingScreenComponentModel> Model;
-};
+		/** The model used by this widget */
+		TSharedPtr<FDMXPixelMappingScreenComponentModel> Model;
+	};
 
-/** Simple view of a screen component, useful if many cells are drawn */
-class SDMXPixelMappingScreenComponentSimplistic
-	: public SCompoundWidget
-{
-public:
-	SLATE_BEGIN_ARGS(SDMXPixelMappingScreenComponentSimplistic)
-	{}
-	SLATE_END_ARGS()
-
-	//~ Begin SCompoundWidget interface
-	void Construct(const FArguments& InArgs, const TSharedRef<FDMXPixelMappingScreenComponentModel>& InModel)
+	/** Simple view of a screen component, useful if many cells are drawn */
+	class SDMXPixelMappingScreenComponentSimplistic
+		: public SCompoundWidget
 	{
-		ChildSlot
-		[
-			SNew(SBox)
-			.HAlign(HAlign_Fill)
-			.VAlign(VAlign_Fill)
+	public:
+		SLATE_BEGIN_ARGS(SDMXPixelMappingScreenComponentSimplistic)
+		{}
+		SLATE_END_ARGS()
+
+		//~ Begin SCompoundWidget interface
+		void Construct(const FArguments& InArgs, const TSharedRef<FDMXPixelMappingScreenComponentModel>& InModel)
+		{
+			ChildSlot
 			[
-				SNew(SScaleBox)
-				.Stretch(EStretch::ScaleToFit)
-				.StretchDirection(EStretchDirection::DownOnly)
+				SNew(SBox)
+				.HAlign(HAlign_Fill)
+				.VAlign(VAlign_Fill)
 				[
-					SNew(SVerticalBox)
-					+ SVerticalBox::Slot()
-					.HAlign(HAlign_Center)
-					.VAlign(VAlign_Center)
-					[
-						SNew(STextBlock)
-						.Font(FCoreStyle::GetDefaultFontStyle("Regular", 24))
-						.ColorAndOpacity(FSlateColor::UseForeground())
-						.Text_Lambda([InModel]()
-							{
-								return FText::Format(LOCTEXT("Num_Pixels", "{0} x {1} pixels"), InModel->GetNumColumns(), InModel->GetNumRows());
-							})
-					]
-					+ SVerticalBox::Slot()
-					.HAlign(HAlign_Center)
-					.VAlign(VAlign_Center)
+					SNew(SScaleBox)
+					.Stretch(EStretch::ScaleToFit)
+					.StretchDirection(EStretchDirection::DownOnly)
 					[
-						SNew(STextBlock)
-						.Font(FCoreStyle::GetDefaultFontStyle("Regular", 24))
-						.ColorAndOpacity(FSlateColor::UseForeground())
-						.Text_Lambda([InModel]()
-							{
-								return FText::FromString(FString::Printf(TEXT("%i.%i"), InModel->GetUniverse(), InModel->GetStartingChannel()));
-							})
+						SNew(SVerticalBox)
+						+ SVerticalBox::Slot()
+						.HAlign(HAlign_Center)
+						.VAlign(VAlign_Center)
+						[
+							SNew(STextBlock)
+							.Font(FCoreStyle::GetDefaultFontStyle("Regular", 24))
+							.ColorAndOpacity(FSlateColor::UseForeground())
+							.Text_Lambda([InModel]()
+								{
+									return FText::Format(LOCTEXT("Num_Pixels", "{0} x {1} pixels"), InModel->GetNumColumns(), InModel->GetNumRows());
+								})
+						]
+						+ SVerticalBox::Slot()
+						.HAlign(HAlign_Center)
+						.VAlign(VAlign_Center)
+						[
+							SNew(STextBlock)
+							.Font(FCoreStyle::GetDefaultFontStyle("Regular", 24))
+							.ColorAndOpacity(FSlateColor::UseForeground())
+							.Text_Lambda([InModel]()
+								{
+									return FText::FromString(FString::Printf(TEXT("%i.%i"), InModel->GetUniverse(), InModel->GetStartingChannel()));
+								})
+						]
 					]
 				]
-			]
-		];
-	}
-	//~ Begin SCompoundWidget interface
+			];
+		}
+		//~ Begin SCompoundWidget interface
 
-private:
-	/** The model used by this widget */
-	TSharedPtr<FDMXPixelMappingScreenComponentModel> Model;
-};
+	private:
+		/** The model used by this widget */
+		TSharedPtr<FDMXPixelMappingScreenComponentModel> Model;
+	};
 
-void SDMXPixelMappingScreenComponent::Construct(const FArguments& InArgs, const TSharedRef<FDMXPixelMappingToolkit>& InToolkit, TWeakObjectPtr<UDMXPixelMappingScreenComponent> ScreenComponent)
-{
-	if (!ScreenComponent.IsValid())
+	void SDMXPixelMappingScreenComponent::Construct(const FArguments& InArgs, const TSharedRef<FDMXPixelMappingToolkit>& InToolkit, TWeakObjectPtr<UDMXPixelMappingScreenComponent> ScreenComponent)
 	{
-		return;
-	}
+		if (!ScreenComponent.IsValid())
+		{
+			return;
+		}
 
-	SetCanTick(true);
+		SetCanTick(true);
 
-	WeakToolkit = InToolkit;
-	Model = MakeShared<FDMXPixelMappingScreenComponentModel>(InToolkit, ScreenComponent);
+		WeakToolkit = InToolkit;
+		Model = MakeShared<FDMXPixelMappingScreenComponentModel>(InToolkit, ScreenComponent);
 
-	ChildSlot
-	[
-		SNew(SBox)
-		.WidthOverride_Lambda([this]()
-			{
-				return Model->GetSize().X;
-			})
-		.HeightOverride_Lambda([this]()
-			{
-				return Model->GetSize().Y;
-			})
+		ChildSlot
 		[
-			SNew(SOverlay)
+			SNew(SBox)
+			.WidthOverride_Lambda([this]()
+				{
+					return Model->GetSize().X;
+				})
+			.HeightOverride_Lambda([this]()
+				{
+					return Model->GetSize().Y;
+				})
+			[
+				SNew(SOverlay)
+
+				+ SOverlay::Slot() 
+					[
+						SAssignNew(ContentBorder, SBorder)
+						.BorderImage(FAppStyle::GetBrush("NoBorder"))
+						.Padding(0.f)
+					]
 
-			+ SOverlay::Slot() 
+				+ SOverlay::Slot() // Required to correctly draw the color around the grid
 				[
-					SAssignNew(ContentBorder, SBorder)
-					.BorderImage(FAppStyle::GetBrush("NoBorder"))
-					.Padding(0.f)
+					SNew(SBorder)
+					.BorderImage(FDMXPixelMappingEditorStyle::Get().GetBrush("DMXPixelMappingEditor.ComponentBorder"))
+					.BorderBackgroundColor_Lambda([this]()
+						{
+							return Model->GetColor();
+						})
 				]
-
-			+ SOverlay::Slot() // Required to correctly draw the color around the grid
-			[
-				SNew(SBorder)
-				.BorderImage(FDMXPixelMappingEditorStyle::Get().GetBrush("DMXPixelMappingEditor.ComponentBorder"))
-				.BorderBackgroundColor_Lambda([this]()
-					{
-						return Model->GetColor();
-					})
 			]
-		]
-	];	
+		];	
 
-	UpdateContent();
-}
+		UpdateContent();
+	}
 
-bool SDMXPixelMappingScreenComponent::Equals(UDMXPixelMappingBaseComponent* Component) const
-{
-	if (Component)
+	bool SDMXPixelMappingScreenComponent::Equals(UDMXPixelMappingBaseComponent* Component) const
 	{
-		return Model->Equals(Component);
+		if (Component)
+		{
+			return Model->Equals(Component);
+		}
+		return false;
 	}
-	return false;
-}
-
-FVector2D SDMXPixelMappingScreenComponent::GetPosition() const
-{
-	return Model->GetPosition();
-}
 
-void SDMXPixelMappingScreenComponent::Tick(const FGeometry& AllottedGeometry, const double InCurrentTime, const float InDeltaTime)
-{
-	const bool bShouldDrawSimplisticView = Model->GetNumColumns() * Model->GetNumRows() > 160;
-	if (bDrawsSimplisticView != bShouldDrawSimplisticView)
+	FVector2D SDMXPixelMappingScreenComponent::GetPosition() const
 	{
-		UpdateContent();
+		return Model->GetPosition();
 	}
-}
 
-void SDMXPixelMappingScreenComponent::UpdateContent()
-{
-	if(Model->GetNumColumns() * Model->GetNumRows() > 160)
+	void SDMXPixelMappingScreenComponent::Tick(const FGeometry& AllottedGeometry, const double InCurrentTime, const float InDeltaTime)
 	{
-		bDrawsSimplisticView = true;
-		ContentBorder->SetContent
-		(
-			SNew(SDMXPixelMappingScreenComponentSimplistic, Model.ToSharedRef())
-		);
+		const bool bShouldDrawSimplisticView = Model->GetNumColumns() * Model->GetNumRows() > 160;
+		if (bDrawsSimplisticView != bShouldDrawSimplisticView)
+		{
+			UpdateContent();
+		}
 	}
-	else
+
+	void SDMXPixelMappingScreenComponent::UpdateContent()
 	{
-		bDrawsSimplisticView = false;
-		ContentBorder->SetContent
-		(
-			SNew(SDMXPixelMappingScreenComponentGrid, Model.ToSharedRef())
-		);
+		if(Model->GetNumColumns() * Model->GetNumRows() > 160)
+		{
+			bDrawsSimplisticView = true;
+			ContentBorder->SetContent
+			(
+				SNew(SDMXPixelMappingScreenComponentSimplistic, Model.ToSharedRef())
+			);
+		}
+		else
+		{
+			bDrawsSimplisticView = false;
+			ContentBorder->SetContent
+			(
+				SNew(SDMXPixelMappingScreenComponentGrid, Model.ToSharedRef())
+			);
+		}
 	}
 }
 
diff --git a/Engine/Plugins/VirtualProduction/DMX/DMXPixelMapping/Source/DMXPixelMappingEditor/Private/Widgets/SDMXPixelMappingOutputComponent.h b/Engine/Plugins/VirtualProduction/DMX/DMXPixelMapping/Source/DMXPixelMappingEditor/Private/Widgets/SDMXPixelMappingOutputComponent.h
--- a/Engine/Plugins/VirtualProduction/DMX/DMXPixelMapping/Source/DMXPixelMappingEditor/Private/Widgets/SDMXPixelMappingOutputComponent.h
+++ b/Engine/Plugins/VirtualProduction/DMX/DMXPixelMapping/Source/DMXPixelMappingEditor/Private/Widgets/SDMXPixelMappingOutputComponent.h
@@ -6,166 +6,133 @@
 #include "Widgets/SCompoundWidget.h"
 #include "Widgets/Layout/SConstraintCanvas.h"
 
-class FDMXPixelMappingOutputComponentModel;
-class FDMXPixelMappingScreenComponentModel;
 class FDMXPixelMappingToolkit;
-class UDMXPixelMappingBaseComponent;
-class UDMXPixelMappingOutputComponent;
-class UDMXPixelMappingScreenComponent;
-
 class SBorder;
 class SBox;
 class STextBlock;
+class UDMXPixelMappingBaseComponent;
+class UDMXPixelMappingOutputComponent;
+class UDMXPixelMappingScreenComponent;
 
-
-/** Interface for Output Component Widgets */
-class IDMXPixelMappingOutputComponentWidgetInterface
-{
-public:
-	/** Destructor */
-	virtual ~IDMXPixelMappingOutputComponentWidgetInterface();
-
-	/** Adds the widgets to a canvas. If it already resides in a canvas it is removed from that first. */
-	void AddToCanvas(const TSharedRef<SConstraintCanvas>& InCanvas);
-
-	/** Removes the widget from the canvas, if it was added to one. */
-	void RemoveFromCanvas();
-
-	/** Returns the actual widget implementation */
-	virtual TSharedRef<SWidget> AsWidget() = 0;
-
-	/** Returns true if the widet equals the component */
-	virtual bool Equals(UDMXPixelMappingBaseComponent* Component) const = 0;
-
-protected:
-	/** Returns the position of the widget */
-	virtual FVector2D GetPosition() const = 0;
-
-private:
-	/** When added to a parent, the canvas it was added to */
-	TSharedPtr<SConstraintCanvas> ParentCanvas;
-
-	/** The canvas slot of the component widget */
-	SConstraintCanvas::FSlot* Slot = nullptr;
-};
-
-
-/** Widget that draws an Output Component. For Screen Component, see SDMXPixelMappingScreenComponent */
-class SDMXPixelMappingOutputComponent
-	: public IDMXPixelMappingOutputComponentWidgetInterface
-	, public SCompoundWidget
+namespace UE::DMX
 {
-public:
-	SLATE_BEGIN_ARGS(SDMXPixelMappingOutputComponent)
-	{}
-	SLATE_END_ARGS()
+	class FDMXPixelMappingOutputComponentModel;
+	class FDMXPixelMappingScreenComponentModel;
 
-	/** Constructs the widget */
-	void Construct(const FArguments& InArgs, const TSharedRef<FDMXPixelMappingToolkit>& InToolkit, TWeakObjectPtr<UDMXPixelMappingOutputComponent> OutputComponent);
 
-	//~ Begin IDMXPixelMappingOutputComponentWidgetInterface
-	virtual TSharedRef<SWidget> AsWidget() override { return AsShared(); };
-	virtual bool Equals(UDMXPixelMappingBaseComponent* Component) const override;
-protected:
-	virtual FVector2D GetPosition() const override;
-	//~ End IDMXPixelMappingOutputComponentWidgetInterface
+	/** Interface for Output Component Widgets */
+	class IDMXPixelMappingOutputComponentWidgetInterface
+	{
+	public:
+		/** Destructor */
+		virtual ~IDMXPixelMappingOutputComponentWidgetInterface();
 
-private:
-	/** Refreshes the widget oo the next tick */
-	void RequestRefresh();
+		/** Adds the widgets to a canvas. If it already resides in a canvas it is removed from that first. */
+		void AddToCanvas(const TSharedRef<SConstraintCanvas>& InCanvas);
 
-	/** Refreshes the widget */
-	void ForceRefresh();
+		/** Removes the widget from the canvas, if it was added to one. */
+		void RemoveFromCanvas();
 
-	/** Creates content */
-	TSharedRef<SWidget> CreateContent();
+		/** Returns the actual widget implementation */
+		virtual TSharedRef<SWidget> AsWidget() = 0;
 
-	/** Creates the child slot that displays the component name, above the child slot */
-	void CreateComponentNameChildSlotAbove();
+		/** Returns true if the widet equals the component */
+		virtual bool Equals(UDMXPixelMappingBaseComponent* Component) const = 0;
 
-	/** Creates the child slot that displays the component name, inside the child slot */
-	void CreateComponentNameChildSlotInside();
+	protected:
+		/** Returns the position of the widget */
+		virtual FVector2D GetPosition() const = 0;
 
-	/** Creates the child slot that displays the Fixture ID of the component */
-	void CreateCellIDChildSlot();
+	private:
+		/** When added to a parent, the canvas it was added to */
+		TSharedPtr<SConstraintCanvas> ParentCanvas;
 
-	/** Creates the child slot that displays info about the patch such as Addresses or the Fixture ID */
-	void CreatePatchInfoChildSlot();
+		/** The canvas slot of the component widget */
+		SConstraintCanvas::FSlot* Slot = nullptr;
+	};
 
-	/** Returns the render transform for this widget */
-	TOptional<FSlateRenderTransform> ComputeRenderTransform() const;
 
-	/** The box that is shown */
-	TSharedPtr<SBox> ComponentBox;
+	/** Widget that draws an Output Component. For Screen Component, see SDMXPixelMappingScreenComponent */
+	class SDMXPixelMappingOutputComponent
+		: public IDMXPixelMappingOutputComponentWidgetInterface
+		, public SCompoundWidget
+	{
+	public:
+		SLATE_BEGIN_ARGS(SDMXPixelMappingOutputComponent)
+		{}
 
-	/** The box that is shown */
-	TSharedPtr<STextBlock> IDTextBlock;
+		SLATE_END_ARGS()
 
-	/** Border for content to display names above the widget */
-	TSharedPtr<SBorder> AboveContentBorder;
+		/** Constructs the widget */
+		void Construct(const FArguments& InArgs, const TSharedRef<FDMXPixelMappingToolkit>& InToolkit, TWeakObjectPtr<UDMXPixelMappingOutputComponent> InOutputComponent);
 
-	/** Border for content to display names in the top row of the widget */
-	TSharedPtr<SBorder> TopContentBorder;
+		//~ Begin IDMXPixelMappingOutputComponentWidgetInterface
+		virtual TSharedRef<SWidget> AsWidget() override { return AsShared(); };
+		virtual bool Equals(UDMXPixelMappingBaseComponent* Component) const override;
 
-	/** Border for content to display names in the middle of the widget */
-	TSharedPtr<SBorder> MiddleContentBorder;
+	protected:
+		virtual FVector2D GetPosition() const override;
+		//~ End IDMXPixelMappingOutputComponentWidgetInterface
 
-	/** Border for content to display names in the top row of the widget */
-	TSharedPtr<SBorder> BottomContentBorder;
+		//~ Begin SWidget Interface
+		virtual int32 OnPaint(const FPaintArgs& Args, const FGeometry& AllottedGeometry, const FSlateRect& MyCullingRect, FSlateWindowElementList& OutDrawElements, int32 LayerId, const FWidgetStyle& InWidgetStyle, bool bParentEnabled) const override;
+		//~ End SWidget Interface
 
-	/** Timer handle for request refresh */
-	FTimerHandle RefreshTimerHandle;
+	private:
+		/** Returns the render transform for this widget */
+		TOptional<FSlateRenderTransform> GetRenderTransform() const;
 
-	/** The model for this widget */
-	TSharedPtr<FDMXPixelMappingOutputComponentModel> Model;
+		/** The model for this widget */
+		TSharedPtr<FDMXPixelMappingOutputComponentModel> Model;
 
-	/** The toolkit that owns */
-	TWeakPtr<FDMXPixelMappingToolkit> WeakToolkit;
-};
+		/** The toolkit that owns */
+		TWeakPtr<FDMXPixelMappingToolkit> WeakToolkit;
+	};
 
-/** Widget that draws a Screen Component. */
-class SDMXPixelMappingScreenComponent
-	: public IDMXPixelMappingOutputComponentWidgetInterface
-	, public SCompoundWidget
-{
-public:
-	SLATE_BEGIN_ARGS(SDMXPixelMappingScreenComponent)
-	{}
-	SLATE_END_ARGS()
+	/** Widget that draws a Screen Component. While the screen component cannot be added to new assets anymore, we still draw it for old ones. */
+	class SDMXPixelMappingScreenComponent
+		: public IDMXPixelMappingOutputComponentWidgetInterface
+		, public SCompoundWidget
+	{
+	public:
+		SLATE_BEGIN_ARGS(SDMXPixelMappingScreenComponent)
+		{}
+		SLATE_END_ARGS()
 
-	/** Constructs the widget */
-	void Construct(const FArguments& InArgs, const TSharedRef<FDMXPixelMappingToolkit>& InToolkit, TWeakObjectPtr<UDMXPixelMappingScreenComponent> ScreenComponent);
+		/** Constructs the widget */
+		void Construct(const FArguments& InArgs, const TSharedRef<FDMXPixelMappingToolkit>& InToolkit, TWeakObjectPtr<UDMXPixelMappingScreenComponent> ScreenComponent);
 
-	//~ Begin IDMXPixelMappingOutputComponentWidgetInterface
-	virtual TSharedRef<SWidget> AsWidget() override { return AsShared(); };
-	virtual bool Equals(UDMXPixelMappingBaseComponent* Component) const override;
-protected:
-	virtual FVector2D GetPosition() const override;
-	//~ End IDMXPixelMappingOutputComponentWidgetInterface
+		//~ Begin IDMXPixelMappingOutputComponentWidgetInterface
+		virtual TSharedRef<SWidget> AsWidget() override { return AsShared(); };
+		virtual bool Equals(UDMXPixelMappingBaseComponent* Component) const override;
+	protected:
+		virtual FVector2D GetPosition() const override;
+		//~ End IDMXPixelMappingOutputComponentWidgetInterface
 
-	//~ Begin SWidget interface
-	virtual void Tick(const FGeometry& AllottedGeometry, const double InCurrentTime, const float InDeltaTime) override;
-	//~ End SWidget interface
+		//~ Begin SWidget interface
+		virtual void Tick(const FGeometry& AllottedGeometry, const double InCurrentTime, const float InDeltaTime) override;
+		//~ End SWidget interface
 
-private:
-	/** Updates the content of the widget */
-	void UpdateContent();
+	private:
+		/** Updates the content of the widget */
+		void UpdateContent();
 
-	/** True if them model currently draws a simplistic view of the component */
-	bool bDrawsSimplisticView = false;
+		/** True if them model currently draws a simplistic view of the component */
+		bool bDrawsSimplisticView = false;
 
-	/** When added to a parent, the canvas it was added to */
-	TSharedPtr<SConstraintCanvas> ParentCanvas;
+		/** When added to a parent, the canvas it was added to */
+		TSharedPtr<SConstraintCanvas> ParentCanvas;
 
-	/** The canvas slot of the component widget */
-	SConstraintCanvas::FSlot* Slot = nullptr;
+		/** The canvas slot of the component widget */
+		SConstraintCanvas::FSlot* Slot = nullptr;
 
-	/** Border that holds the content */
-	TSharedPtr<SBorder> ContentBorder;
+		/** Border that holds the content */
+		TSharedPtr<SBorder> ContentBorder;
 
-	/** The model for this widget */
-	TSharedPtr<FDMXPixelMappingScreenComponentModel> Model;
+		/** The model for this widget */
+		TSharedPtr<FDMXPixelMappingScreenComponentModel> Model;
 
-	/** The toolkit that owns */
-	TWeakPtr<FDMXPixelMappingToolkit> WeakToolkit;
-};
+		/** The toolkit that owns */
+		TWeakPtr<FDMXPixelMappingToolkit> WeakToolkit;
+	};
+}
\ No newline at end of file
diff --git a/Engine/Plugins/VirtualProduction/DMX/DMXPixelMapping/Source/DMXPixelMappingEditor/Private/Widgets/SDMXPixelMappingOutputComponentText.cpp b/Engine/Plugins/VirtualProduction/DMX/DMXPixelMapping/Source/DMXPixelMappingEditor/Private/Widgets/SDMXPixelMappingOutputComponentText.cpp
new file mode 100644
--- /dev/null
+++ b/Engine/Plugins/VirtualProduction/DMX/DMXPixelMapping/Source/DMXPixelMappingEditor/Private/Widgets/SDMXPixelMappingOutputComponentText.cpp
@@ -0,0 +1,173 @@
+// Copyright Epic Games, Inc. All Rights Reserved.
+
+#include "Widgets/SDMXPixelMappingOutputComponentText.h"
+
+#include "Fonts/FontMeasure.h"
+#include "Framework/Application/SlateApplication.h"
+#include "Styling/CoreStyle.h"
+#include "ViewModels/DMXPixelMappingOutputComponentModel.h"
+
+
+void SDMXPixelMappingOutputComponentText::Construct(const FArguments& InArgs, const TSharedRef<FDMXPixelMappingToolkit>& InToolkit, TWeakObjectPtr<UDMXPixelMappingOutputComponent> OutputComponent)
+{
+	if (!OutputComponent.IsValid())
+	{
+		return;
+	}
+
+	WeakToolkit = InToolkit;
+	Model = MakeShared<UE::DMX::FDMXPixelMappingOutputComponentModel>(InToolkit, OutputComponent);
+	
+	FontSize = InArgs._FontSize;
+
+	SetRenderTransform(TAttribute<TOptional<FSlateRenderTransform>>::CreateSP(this, &SDMXPixelMappingOutputComponentText::GetRenderTransform));
+	SetRenderTransformPivot(FVector2D(0.5, 0.5));
+
+	// Don't clip if the name has to be painted above the component box, otherwise clip. 
+	const EWidgetClipping WidgetClipping = Model->ShouldDrawNameAbove() ? EWidgetClipping::Inherit : EWidgetClipping::ClipToBoundsAlways;
+	SetClipping(WidgetClipping);
+}
+
+int32 SDMXPixelMappingOutputComponentText::OnPaint(
+	const FPaintArgs& Args,
+	const FGeometry& AllottedGeometry,
+	const FSlateRect& MyCullingRect,
+	FSlateWindowElementList& OutDrawElements,
+	int32 LayerId,
+	const FWidgetStyle& InWidgetStyle,
+	bool bParentEnabled) const
+{
+	if (!Model->ShouldDraw())
+	{
+		return LayerId;
+	}
+	LayerId = SCompoundWidget::OnPaint(Args, AllottedGeometry, MyCullingRect, OutDrawElements, LayerId, InWidgetStyle, bParentEnabled);
+
+	// Selectively paint labels
+	const FVector2f Scale = AllottedGeometry.GetAccumulatedRenderTransform().GetMatrix().GetScale().GetVector();
+	if (Model->ShouldDrawName())
+	{
+		OnPaintComponentName(Args, AllottedGeometry, OutDrawElements, LayerId, Scale);
+	}
+
+	if (Model->ShouldDrawCellID())
+	{
+		OnPaintCellID(Args, AllottedGeometry, OutDrawElements, LayerId, Scale);
+	}
+
+	if (Model->ShouldDrawPatchInfo())
+	{
+		OnPaintPatchInfo(Args, AllottedGeometry, OutDrawElements, LayerId, Scale);
+	}
+
+	return LayerId + 1;
+}
+
+TOptional<FSlateRenderTransform> SDMXPixelMappingOutputComponentText::GetRenderTransform() const
+{
+	const FQuat2D Quaternion = Model->GetQuaternion();
+
+	// The text always should face upwards to be readable. Hence,
+	// if the rotation goes > 90° or < -90°, draw the widget inverted.
+	if (Quaternion.GetVector().X < 0.0)
+	{
+		return FSlateRenderTransform(FQuat2D(UE_PI), FVector2D::ZeroVector);
+	}
+
+	return TOptional<FSlateRenderTransform>();
+}
+
+void SDMXPixelMappingOutputComponentText::OnPaintComponentName(const FPaintArgs& Args, const FGeometry& AllottedGeometry, FSlateWindowElementList& OutDrawElements, int32 LayerId, const FVector2f& Scale) const
+{
+	const float Margin = 8.f / Scale.Y;
+	const FVector2f LocalSize = AllottedGeometry.GetLocalSize();
+
+	// Upscale the font so it stays the same size, regardless of scaling (zoom)
+	const float ScaledFontSize = FontSize / Scale.Y;
+
+	const float MaxFontSize = LocalSize.Y - ScaledFontSize - Margin * 2.f;
+	const float DesiredFontSize = ScaledFontSize < MaxFontSize ? ScaledFontSize : MaxFontSize;
+	const float ClampedFontSize = FMath::Max(UE_DELTA, DesiredFontSize);
+
+	const float OffsetX = Margin;
+	const float OffsetY = Model->ShouldDrawNameAbove() ? -ScaledFontSize - Margin : Margin;
+
+	const FSlateLayoutTransform LayoutTransform(FVector2f(OffsetX, OffsetY));
+
+	FSlateDrawElement::MakeText(
+		OutDrawElements,
+		LayerId,
+		AllottedGeometry.ToPaintGeometry(LayoutTransform),
+		Model->GetName(),
+		FCoreStyle::GetDefaultFontStyle("Regular", ClampedFontSize));
+}
+
+void SDMXPixelMappingOutputComponentText::OnPaintCellID(const FPaintArgs& Args, const FGeometry& AllottedGeometry, FSlateWindowElementList& OutDrawElements, int32 LayerId, const FVector2f& Scale) const
+{
+	const float Margin = 8.f / Scale.Y;
+	const FVector2f LocalSize = AllottedGeometry.GetLocalSize();
+
+	// Upscale the font so it stays the same size, regardless of scaling (zoom)
+	const float ScaledFontSize = FontSize / Scale.Y;
+
+	const float MaxFontSize = LocalSize.Y - ScaledFontSize - Margin * 2.f;
+	const float DesiredFontSize = ScaledFontSize < MaxFontSize ? ScaledFontSize : MaxFontSize;
+	const float ClampedFontSize = FMath::Max(UE_DELTA, DesiredFontSize);
+
+	const FText CellIDText = Model->GetCellIDText();
+	const FSlateFontInfo Font = FCoreStyle::GetDefaultFontStyle("Regular", ScaledFontSize);
+
+	const TSharedRef<FSlateFontMeasure>& FontMeasureService = FSlateApplication::Get().GetRenderer()->GetFontMeasureService();
+	const FVector2D TextSize = FontMeasureService->Measure(CellIDText, Font) / Scale;
+
+	const FSlateLayoutTransform LayoutTransform(FVector2f(LocalSize.X / 2.f - TextSize.X, LocalSize.Y / 2.f - TextSize.Y));
+
+	FSlateDrawElement::MakeText(
+		OutDrawElements,
+		LayerId,
+		AllottedGeometry.ToPaintGeometry(LayoutTransform),
+		Model->GetCellIDText(),
+		Font,
+		ESlateDrawEffect::None,
+		FLinearColor::White.CopyWithNewOpacity(.6f));
+}
+
+void SDMXPixelMappingOutputComponentText::OnPaintPatchInfo(const FPaintArgs& Args, const FGeometry& AllottedGeometry, FSlateWindowElementList& OutDrawElements, int32 LayerId, const FVector2f& Scale) const
+{
+	const float Margin = 8.f / Scale.Y;
+	const float LocalSizeY = AllottedGeometry.GetLocalSize().Y;
+
+	// Upscale the font so it stays the same size, regardless of scaling (zoom)
+	const float ScaledFontSize = FontSize / Scale.Y;
+	const float LineSpacing = 4.f / Scale.Y;
+
+	// Find a font size that fits in the lower third of this section. Consider two lines of text are painted
+	const float MaxFontSize = LocalSizeY - ScaledFontSize * 2.f - LineSpacing - Margin * 2.f;
+	const float DesiredFontSize = ScaledFontSize < MaxFontSize ? ScaledFontSize : MaxFontSize;
+	const float ClampedFontSize = FMath::Max(UE_DELTA, DesiredFontSize);
+
+	// Arrange both text elements in the lower section, align bottom
+	const float FixtureIDTextOffsetY = LocalSizeY - ClampedFontSize * 2.f - Margin - LineSpacing;
+	const FSlateLayoutTransform FixtureIDLayoutTransform(FVector2f(Margin, FixtureIDTextOffsetY));
+
+	FSlateDrawElement::MakeText(
+		OutDrawElements,
+		LayerId,
+		AllottedGeometry.ToPaintGeometry(FixtureIDLayoutTransform),
+		Model->GetFixtureIDText(),
+		FCoreStyle::GetDefaultFontStyle("Regular", ClampedFontSize),
+		ESlateDrawEffect::None,
+		FLinearColor::Green);
+
+	const float AddressesTextOffsetY = LocalSizeY - ClampedFontSize - Margin;
+	const FSlateLayoutTransform AddressesLayoutTransform(FVector2f(Margin, AddressesTextOffsetY));
+
+	FSlateDrawElement::MakeText(
+		OutDrawElements,
+		LayerId,
+		AllottedGeometry.ToPaintGeometry(AddressesLayoutTransform),
+		Model->GetAddressesText(),
+		FCoreStyle::GetDefaultFontStyle("Regular", ClampedFontSize),
+		ESlateDrawEffect::None,
+		FLinearColor::Green);
+}
diff --git a/Engine/Plugins/VirtualProduction/DMX/DMXPixelMapping/Source/DMXPixelMappingEditor/Private/Widgets/SDMXPixelMappingOutputComponentText.h b/Engine/Plugins/VirtualProduction/DMX/DMXPixelMapping/Source/DMXPixelMappingEditor/Private/Widgets/SDMXPixelMappingOutputComponentText.h
new file mode 100644
--- /dev/null
+++ b/Engine/Plugins/VirtualProduction/DMX/DMXPixelMapping/Source/DMXPixelMappingEditor/Private/Widgets/SDMXPixelMappingOutputComponentText.h
@@ -0,0 +1,55 @@
+// Copyright Epic Games, Inc. All Rights Reserved.
+
+#pragma once
+
+#include "Widgets/SCompoundWidget.h"
+
+class FDMXPixelMappingToolkit;
+class UDMXPixelMappingOutputComponent;
+namespace UE::DMX { class FDMXPixelMappingOutputComponentModel; }
+
+
+/** Draws the text in an output component widget */
+class SDMXPixelMappingOutputComponentText
+	: public SCompoundWidget
+{
+
+public:
+	SLATE_BEGIN_ARGS(SDMXPixelMappingOutputComponentText)
+		: _FontSize(8)
+		{}
+		
+		SLATE_ARGUMENT(int32, FontSize)
+
+	SLATE_END_ARGS()
+
+	/** Constructs this widget */
+	void Construct(const FArguments& InArgs, const TSharedRef<FDMXPixelMappingToolkit>& InToolkit, TWeakObjectPtr<UDMXPixelMappingOutputComponent> OutputComponent);
+
+	/** The model for this widget */
+	TSharedPtr<UE::DMX::FDMXPixelMappingOutputComponentModel> Model;
+
+	/** The toolkit that owns */
+	TWeakPtr<FDMXPixelMappingToolkit> WeakToolkit;
+
+protected:
+	//~ Begin SWidget Interface
+	virtual int32 OnPaint(const FPaintArgs& Args, const FGeometry& AllottedGeometry, const FSlateRect& MyCullingRect, FSlateWindowElementList& OutDrawElements, int32 LayerId, const FWidgetStyle& InWidgetStyle, bool bParentEnabled) const override;
+	//~ End SWidget Interface
+
+private:
+	/** Returns the render transform for this widget */
+	TOptional<FSlateRenderTransform> GetRenderTransform() const;
+
+	/** Helper to paint the component name */
+	void OnPaintComponentName(const FPaintArgs& Args, const FGeometry& AllottedGeometry, FSlateWindowElementList& OutDrawElements, int32 LayerId, const FVector2f& Scale) const;
+
+	/** Helper to paint the cell ID */
+	void OnPaintCellID(const FPaintArgs& Args, const FGeometry& AllottedGeometry, FSlateWindowElementList& OutDrawElements, int32 LayerId, const FVector2f& Scale) const;
+
+	/** Helper to paint the patch info */
+	void OnPaintPatchInfo(const FPaintArgs& Args, const FGeometry& AllottedGeometry, FSlateWindowElementList& OutDrawElements, int32 LayerId, const FVector2f& Scale) const;
+
+	// Slate args
+	int32 FontSize;
+};
diff --git a/Engine/Plugins/VirtualProduction/DMX/DMXPixelMapping/Source/DMXPixelMappingRuntime/Internal/Components/DMXPixelMappingComponentGeometryCache.h b/Engine/Plugins/VirtualProduction/DMX/DMXPixelMapping/Source/DMXPixelMappingRuntime/Internal/Components/DMXPixelMappingComponentGeometryCache.h
--- a/Engine/Plugins/VirtualProduction/DMX/DMXPixelMapping/Source/DMXPixelMappingRuntime/Internal/Components/DMXPixelMappingComponentGeometryCache.h
+++ b/Engine/Plugins/VirtualProduction/DMX/DMXPixelMapping/Source/DMXPixelMappingRuntime/Internal/Components/DMXPixelMappingComponentGeometryCache.h
@@ -49,8 +49,8 @@ public:
 	/** Returns rotation as sin cos */
 	void GetSinCos(double& OutSin, double& OutCos) const;
 
-	/** Returns the edges of the rectangle, considering rotation. */
-	void GetEdges(FVector2D& A, FVector2D& B, FVector2D& C, FVector2D& D) const;
+	/** Returns the edges of the rectangle, absolute, rotated, clockwise order */
+	void GetEdgesAbsolute(FVector2D& A, FVector2D& B, FVector2D& C, FVector2D& D) const;
 
 	/** Propagonates position changes to children. */
 	void PropagonatePositionChangesToChildren(const FVector2D& Translation);
diff --git a/Engine/Plugins/VirtualProduction/DMX/DMXPixelMapping/Source/DMXPixelMappingRuntime/Private/Components/DMXPixelMappingComponentGeometryCache.cpp b/Engine/Plugins/VirtualProduction/DMX/DMXPixelMapping/Source/DMXPixelMappingRuntime/Private/Components/DMXPixelMappingComponentGeometryCache.cpp
--- a/Engine/Plugins/VirtualProduction/DMX/DMXPixelMapping/Source/DMXPixelMappingRuntime/Private/Components/DMXPixelMappingComponentGeometryCache.cpp
+++ b/Engine/Plugins/VirtualProduction/DMX/DMXPixelMapping/Source/DMXPixelMappingRuntime/Private/Components/DMXPixelMappingComponentGeometryCache.cpp
@@ -93,7 +93,7 @@ void FDMXPixelMappingComponentGeometryCache::GetSinCos(double& OutSin, double& O
 	OutCos = Cos;
 }
 
-void FDMXPixelMappingComponentGeometryCache::GetEdges(FVector2D& A, FVector2D& B, FVector2D& C, FVector2D& D) const
+void FDMXPixelMappingComponentGeometryCache::GetEdgesAbsolute(FVector2D& A, FVector2D& B, FVector2D& C, FVector2D& D) const
 {
 	const FVector2D Pivot = Position + Size / 2.0;
 
diff --git a/Engine/Plugins/VirtualProduction/DMX/DMXPixelMapping/Source/DMXPixelMappingRuntime/Private/Components/DMXPixelMappingFixtureGroupItemComponent.cpp b/Engine/Plugins/VirtualProduction/DMX/DMXPixelMapping/Source/DMXPixelMappingRuntime/Private/Components/DMXPixelMappingFixtureGroupItemComponent.cpp
--- a/Engine/Plugins/VirtualProduction/DMX/DMXPixelMapping/Source/DMXPixelMappingRuntime/Private/Components/DMXPixelMappingFixtureGroupItemComponent.cpp
+++ b/Engine/Plugins/VirtualProduction/DMX/DMXPixelMapping/Source/DMXPixelMappingRuntime/Private/Components/DMXPixelMappingFixtureGroupItemComponent.cpp
@@ -376,7 +376,7 @@ bool UDMXPixelMappingFixtureGroupItemComponent::IsOverParent() const
 		FVector2D B;
 		FVector2D C;
 		FVector2D D;
-		CachedGeometry.GetEdges(A, B, C, D);
+		CachedGeometry.GetEdgesAbsolute(A, B, C, D);
 
 		return
 			Parent->IsOverPosition(A) &&
diff --git a/Engine/Plugins/VirtualProduction/DMX/DMXPixelMapping/Source/DMXPixelMappingRuntime/Private/Components/DMXPixelMappingMatrixCellComponent.cpp b/Engine/Plugins/VirtualProduction/DMX/DMXPixelMapping/Source/DMXPixelMappingRuntime/Private/Components/DMXPixelMappingMatrixCellComponent.cpp
--- a/Engine/Plugins/VirtualProduction/DMX/DMXPixelMapping/Source/DMXPixelMappingRuntime/Private/Components/DMXPixelMappingMatrixCellComponent.cpp
+++ b/Engine/Plugins/VirtualProduction/DMX/DMXPixelMapping/Source/DMXPixelMappingRuntime/Private/Components/DMXPixelMappingMatrixCellComponent.cpp
@@ -87,7 +87,7 @@ bool UDMXPixelMappingMatrixCellComponent::IsOverParent() const
 		FVector2D B;
 		FVector2D C;
 		FVector2D D;
-		CachedGeometry.GetEdges(A, B, C, D);
+		CachedGeometry.GetEdgesAbsolute(A, B, C, D);
 
 		return
 			Parent->IsOverPosition(A) &&
diff --git a/Engine/Plugins/VirtualProduction/DMX/DMXPixelMapping/Source/DMXPixelMappingRuntime/Private/Components/DMXPixelMappingMatrixComponent.cpp b/Engine/Plugins/VirtualProduction/DMX/DMXPixelMapping/Source/DMXPixelMappingRuntime/Private/Components/DMXPixelMappingMatrixComponent.cpp
--- a/Engine/Plugins/VirtualProduction/DMX/DMXPixelMapping/Source/DMXPixelMappingRuntime/Private/Components/DMXPixelMappingMatrixComponent.cpp
+++ b/Engine/Plugins/VirtualProduction/DMX/DMXPixelMapping/Source/DMXPixelMappingRuntime/Private/Components/DMXPixelMappingMatrixComponent.cpp
@@ -217,7 +217,7 @@ bool UDMXPixelMappingMatrixComponent::IsOverParent() const
 		FVector2D B;
 		FVector2D C;
 		FVector2D D;
-		CachedGeometry.GetEdges(A, B, C, D);
+		CachedGeometry.GetEdgesAbsolute(A, B, C, D);
 
 		return
 			Parent->IsOverPosition(A) &&
diff --git a/Engine/Plugins/VirtualProduction/DMX/DMXPixelMapping/Source/DMXPixelMappingRuntime/Private/Components/DMXPixelMappingOutputComponent.cpp b/Engine/Plugins/VirtualProduction/DMX/DMXPixelMapping/Source/DMXPixelMappingRuntime/Private/Components/DMXPixelMappingOutputComponent.cpp
--- a/Engine/Plugins/VirtualProduction/DMX/DMXPixelMapping/Source/DMXPixelMappingRuntime/Private/Components/DMXPixelMappingOutputComponent.cpp
+++ b/Engine/Plugins/VirtualProduction/DMX/DMXPixelMapping/Source/DMXPixelMappingRuntime/Private/Components/DMXPixelMappingOutputComponent.cpp
@@ -111,7 +111,7 @@ bool UDMXPixelMappingOutputComponent::IsOverPosition(const FVector2D& OtherPosit
 	FVector2D B;
 	FVector2D C;
 	FVector2D D;
-	CachedGeometry.GetEdges(A, B, C, D);
+	CachedGeometry.GetEdgesAbsolute(A, B, C, D);
 
 	const FVector2D AM = OtherPosition - A;
 	const FVector2D AB = B - A;
@@ -138,7 +138,7 @@ bool UDMXPixelMappingOutputComponent::OverlapsComponent(UDMXPixelMappingOutputCo
 		FVector2D B;
 		FVector2D C;
 		FVector2D D;
-		CachedGeometry.GetEdges(A, B, C, D);
+		CachedGeometry.GetEdgesAbsolute(A, B, C, D);
 
 		return
 			Other->IsOverPosition(A) ||
@@ -184,6 +184,11 @@ void UDMXPixelMappingOutputComponent::SetPositionRotated(FVector2D NewRotatedPos
 #endif
 }
 
+void UDMXPixelMappingOutputComponent::GetEdges(FVector2D& A, FVector2D& B, FVector2D& C, FVector2D& D) const
+{
+	return CachedGeometry.GetEdgesAbsolute(A, B, C, D);
+}
+
 void UDMXPixelMappingOutputComponent::SetSize(const FVector2D& NewSize)
 {
 	SizeX = NewSize.X;
diff --git a/Engine/Plugins/VirtualProduction/DMX/DMXPixelMapping/Source/DMXPixelMappingRuntime/Public/Components/DMXPixelMappingBaseComponent.h b/Engine/Plugins/VirtualProduction/DMX/DMXPixelMapping/Source/DMXPixelMappingRuntime/Public/Components/DMXPixelMappingBaseComponent.h
--- a/Engine/Plugins/VirtualProduction/DMX/DMXPixelMapping/Source/DMXPixelMappingRuntime/Public/Components/DMXPixelMappingBaseComponent.h
+++ b/Engine/Plugins/VirtualProduction/DMX/DMXPixelMapping/Source/DMXPixelMappingRuntime/Public/Components/DMXPixelMappingBaseComponent.h
@@ -44,7 +44,6 @@ public:
 	static FDMXPixelMappingOnComponentRenamed& GetOnComponentRenamed();
 	
 	//~ Begin UObject interface
-protected:
 	virtual void PostRename(UObject* OldOuter, const FName OldName) override;
 	virtual void Serialize(FArchive& Ar) override;
 #if WITH_EDITOR
@@ -52,7 +51,6 @@ protected:
 #endif // WITH_EDITOR
 	//~ End UObject interface
 
-public:
 	/**
 	 * Should log properties that were changed in underlying fixture patch or fixture type
 	 *
@@ -135,19 +133,6 @@ public:
 		return FoundObject;
 	}
 
-	/** DEPRECATED 4.27  */
-	template <typename TComponentClass>
-	UE_DEPRECATED(4.27, "Use ForEachChildOfClass in favor of a clearer name instead.")
-	void ForEachComponentOfClass(TComponentPredicateType<TComponentClass> Predicate, bool bIsRecursive)
-	{
-		ForEachChild([&Predicate](UDMXPixelMappingBaseComponent* InComponent) {
-			if (TComponentClass* CastComponent = Cast<TComponentClass>(InComponent))
-			{
-				Predicate(CastComponent);
-			}
-		}, bIsRecursive);
-	}
-
 	/** 
 	 * Loop through all templated child class by given Predicate
 	 *
@@ -201,16 +186,6 @@ public:
 	UFUNCTION(BlueprintCallable, Category = "DMX|PixelMapping")
 	virtual void RenderAndSendDMX() {};
 
-#if WITH_EDITORONLY_DATA
-	UPROPERTY()
-	bool bExpanded = true;
-#endif
-
-public:
-	/*----------------------------------------------------------
-		Public static functions
-	----------------------------------------------------------*/
-
 	/**
 	 * Recursively looking for the first parent by given Class
 	 *
@@ -251,11 +226,15 @@ public:
 	/** Set the components parent */
 	void SetParent(const TWeakObjectPtr<UDMXPixelMappingBaseComponent>& NewParent) { WeakParent = NewParent; }
 #endif 
-
 	/** Parent component */
 	UPROPERTY(Meta = (DeprecatedProperty, DeprecationMessage = "Leads to entangled references. Use GetParent() or WeakParent instead."))
 	TObjectPtr<UDMXPixelMappingBaseComponent> Parent_DEPRECATED;
 
+#if WITH_EDITORONLY_DATA
+	UPROPERTY()
+	bool bExpanded = true;
+#endif
+
 protected:
 	/** Called when the component was added to a parent */
 	virtual void NotifyAddedToParent();
diff --git a/Engine/Plugins/VirtualProduction/DMX/DMXPixelMapping/Source/DMXPixelMappingRuntime/Public/Components/DMXPixelMappingOutputComponent.h b/Engine/Plugins/VirtualProduction/DMX/DMXPixelMapping/Source/DMXPixelMappingRuntime/Public/Components/DMXPixelMappingOutputComponent.h
--- a/Engine/Plugins/VirtualProduction/DMX/DMXPixelMapping/Source/DMXPixelMappingRuntime/Public/Components/DMXPixelMappingOutputComponent.h
+++ b/Engine/Plugins/VirtualProduction/DMX/DMXPixelMapping/Source/DMXPixelMappingRuntime/Public/Components/DMXPixelMappingOutputComponent.h
@@ -98,6 +98,9 @@ public:
 	/** Sets the the absolute position with rotation. */
 	virtual void SetPositionRotated(FVector2D NewRotatedPosition);
 
+	/** Returns the edges of the component, absolute, rotated, clockwise order */
+	void GetEdges(FVector2D& A, FVector2D& B, FVector2D& C, FVector2D& D) const;
+
 	/** Get the absolute size */
 	FVector2D GetSize() const;
 
