DMX - Rework Pixel Mapping menus, add new options, add a context menu

diff --git a/Engine/Plugins/VirtualProduction/DMX/DMXPixelMapping/Source/DMXPixelMappingEditor/Private/DMXPixelMappingEditorCommands.cpp b/Engine/Plugins/VirtualProduction/DMX/DMXPixelMapping/Source/DMXPixelMappingEditor/Private/DMXPixelMappingEditorCommands.cpp
--- a/Engine/Plugins/VirtualProduction/DMX/DMXPixelMapping/Source/DMXPixelMappingEditor/Private/DMXPixelMappingEditorCommands.cpp
+++ b/Engine/Plugins/VirtualProduction/DMX/DMXPixelMapping/Source/DMXPixelMappingEditor/Private/DMXPixelMappingEditorCommands.cpp
@@ -14,8 +14,7 @@ FDMXPixelMappingEditorCommands::FDMXPixelMappingEditorCommands()
 		NAME_None,
 		FAppStyle::GetAppStyleSetName()
 	)
-{
-}
+{}
 
 void FDMXPixelMappingEditorCommands::RegisterCommands()
 {
@@ -23,8 +22,8 @@ void FDMXPixelMappingEditorCommands::RegisterCommands()
 	UI_COMMAND(PlayDMX, "Play DMX", "Starts to send DMX", EUserInterfaceActionType::Button, FInputChord());
 	UI_COMMAND(StopPlayingDMX, "Stop Playing DMX", "Stops playing DMX.", EUserInterfaceActionType::Button, FInputChord());
 
-	UI_COMMAND(ResizeMode, "Resize Mode", "Resizes components when transform handles are being dragged", EUserInterfaceActionType::Button, FInputChord());
-	UI_COMMAND(RotateMode, "Rotate Mode", "Rotates components when transform handles are being dragged", EUserInterfaceActionType::Button, FInputChord());
+	UI_COMMAND(EnableResizeMode, "Resize Mode", "Resizes components when transform handles are being dragged", EUserInterfaceActionType::RadioButton, FInputChord());
+	UI_COMMAND(EnableRotateMode, "Rotate Mode", "Rotates components when transform handles are being dragged", EUserInterfaceActionType::RadioButton, FInputChord());
 
 	UI_COMMAND(ToggleGridSnapping, "Toggle Grid Snapping", "Enables/disables grid snapping", EUserInterfaceActionType::Button, FInputChord());
 
@@ -37,6 +36,7 @@ void FDMXPixelMappingEditorCommands::RegisterCommands()
 	UI_COMMAND(ToggleShowComponentNames, "Show Component Names", "Sets if the name of components are displayed. It is recommended that this is turned off when pixel mapping large quantities of fixtures.", EUserInterfaceActionType::ToggleButton, FInputChord(EModifierKey::Alt, EKeys::S));
 	UI_COMMAND(ToggleShowPatchInfo, "Show Patch Info", "Sets if information about the pach is displayed. It is recommended that this is turned off when pixel mapping large quantities of fixtures.", EUserInterfaceActionType::ToggleButton, FInputChord(EModifierKey::Alt, EKeys::D));
 	UI_COMMAND(ToggleShowCellIDs, "Show Cell IDs", "Sets if the cell IDs of matrix cells are displayed. It is recommended that this is turned off when pixel mapping large quantities of fixtures.", EUserInterfaceActionType::ToggleButton, FInputChord(EModifierKey::Alt, EKeys::F));
+	UI_COMMAND(ToggleShowPivot, "Show Pivot", "Sets if the pivot is displayed for selected components.", EUserInterfaceActionType::ToggleButton, FInputChord());
 }
 
 #undef LOCTEXT_NAMESPACE 
diff --git a/Engine/Plugins/VirtualProduction/DMX/DMXPixelMapping/Source/DMXPixelMappingEditor/Private/DMXPixelMappingEditorCommands.h b/Engine/Plugins/VirtualProduction/DMX/DMXPixelMapping/Source/DMXPixelMappingEditor/Private/DMXPixelMappingEditorCommands.h
--- a/Engine/Plugins/VirtualProduction/DMX/DMXPixelMapping/Source/DMXPixelMappingEditor/Private/DMXPixelMappingEditorCommands.h
+++ b/Engine/Plugins/VirtualProduction/DMX/DMXPixelMapping/Source/DMXPixelMappingEditor/Private/DMXPixelMappingEditorCommands.h
@@ -19,8 +19,8 @@ public:
 	TSharedPtr<FUICommandInfo> StopPlayingDMX;
 
 	// Designer related
-	TSharedPtr<FUICommandInfo> ResizeMode;
-	TSharedPtr<FUICommandInfo> RotateMode;
+	TSharedPtr<FUICommandInfo> EnableResizeMode;
+	TSharedPtr<FUICommandInfo> EnableRotateMode;
 
 	TSharedPtr<FUICommandInfo> ToggleGridSnapping;
 
@@ -33,4 +33,5 @@ public:
 	TSharedPtr<FUICommandInfo> ToggleShowComponentNames;
 	TSharedPtr<FUICommandInfo> ToggleShowPatchInfo;
 	TSharedPtr<FUICommandInfo> ToggleShowCellIDs;
+	TSharedPtr<FUICommandInfo> ToggleShowPivot;
 };
diff --git a/Engine/Plugins/VirtualProduction/DMX/DMXPixelMapping/Source/DMXPixelMappingEditor/Private/DMXPixelMappingToolbar.cpp b/Engine/Plugins/VirtualProduction/DMX/DMXPixelMapping/Source/DMXPixelMappingEditor/Private/DMXPixelMappingToolbar.cpp
--- a/Engine/Plugins/VirtualProduction/DMX/DMXPixelMapping/Source/DMXPixelMappingEditor/Private/DMXPixelMappingToolbar.cpp
+++ b/Engine/Plugins/VirtualProduction/DMX/DMXPixelMapping/Source/DMXPixelMappingEditor/Private/DMXPixelMappingToolbar.cpp
@@ -59,56 +59,6 @@ void FDMXPixelMappingToolbar::BuildToolbarCallback(FToolBarBuilder& ToolbarBuild
 			FName(TEXT("Stop Playing DMX")));
 	}
 	ToolbarBuilder.EndSection();
-
-	ToolbarBuilder.BeginSection("Layout");
-	{
-		ToolbarBuilder.AddComboButton(
-			FUIAction(),
-			FOnGetContent::CreateSP(this, &FDMXPixelMappingToolbar::GenerateLayoutMenu),
-			LOCTEXT("PixelMappingToolbarLayoutSettingsLabel", "Layout"),
-			LOCTEXT("PixelMappingToolbarLayoutSettingsTooltip", "Layout Settings for Pixel Mapping Editors"),
-			FSlateIcon(FAppStyle::GetAppStyleSetName(), "Icons.Layout")
-		);
-	}
-	ToolbarBuilder.EndSection();
-}
-
-TSharedRef<SWidget> FDMXPixelMappingToolbar::GenerateLayoutMenu()
-{
-	TSharedPtr<FDMXPixelMappingToolkit> Toolkit = ToolkitWeakPtr.Pin();
-	check(Toolkit.IsValid());
-
-	constexpr bool bShouldCloseMenuAfterSelection = true;
-	FMenuBuilder MenuBuilder(bShouldCloseMenuAfterSelection, Toolkit->GetToolkitCommands());
-	
-	MenuBuilder.BeginSection("Actions", LOCTEXT("ActionsSection", "Actions"));
-	{
-		MenuBuilder.AddMenuEntry(FDMXPixelMappingEditorCommands::Get().SizeComponentToTexture);
-	}
-	MenuBuilder.EndSection();
-
-	MenuBuilder.BeginSection("EditorSettings", LOCTEXT("EditorSettingsSection", "Editor Settings"));
-	{
-		MenuBuilder.AddMenuEntry(FDMXPixelMappingEditorCommands::Get().ToggleScaleChildrenWithParent);
-		MenuBuilder.AddMenuEntry(FDMXPixelMappingEditorCommands::Get().ToggleAlwaysSelectGroup);
-	}
-	MenuBuilder.EndSection();
-
-	MenuBuilder.BeginSection("DisplaySettings", LOCTEXT("DisplaySettingsSection", "Display Settings"));
-	{
-		MenuBuilder.AddMenuEntry(FDMXPixelMappingEditorCommands::Get().ToggleShowComponentNames);
-		MenuBuilder.AddMenuEntry(FDMXPixelMappingEditorCommands::Get().ToggleShowPatchInfo);
-		MenuBuilder.AddMenuEntry(FDMXPixelMappingEditorCommands::Get().ToggleShowMatrixCells);
-		MenuBuilder.AddMenuEntry(FDMXPixelMappingEditorCommands::Get().ToggleShowCellIDs);
-	}
-	MenuBuilder.EndSection();
-
-	MenuBuilder.BeginSection("EditorSettings", LOCTEXT("LayoutSettingsSection", "Layout Settings"));
-	{
-		MenuBuilder.AddMenuEntry(FDMXPixelMappingEditorCommands::Get().ToggleApplyLayoutScriptWhenLoaded);
-	}
-
-	return MenuBuilder.MakeWidget();
 }
 
 #undef LOCTEXT_NAMESPACE
diff --git a/Engine/Plugins/VirtualProduction/DMX/DMXPixelMapping/Source/DMXPixelMappingEditor/Private/DragDrop/DMXPixelMappingDragDropOp.cpp b/Engine/Plugins/VirtualProduction/DMX/DMXPixelMapping/Source/DMXPixelMappingEditor/Private/DragDrop/DMXPixelMappingDragDropOp.cpp
--- a/Engine/Plugins/VirtualProduction/DMX/DMXPixelMapping/Source/DMXPixelMappingEditor/Private/DragDrop/DMXPixelMappingDragDropOp.cpp
+++ b/Engine/Plugins/VirtualProduction/DMX/DMXPixelMapping/Source/DMXPixelMappingEditor/Private/DragDrop/DMXPixelMappingDragDropOp.cpp
@@ -116,21 +116,15 @@ FVector2D FDMXPixelMappingDragDropOp::ComputeGridSnapPosition(const FVector2D& D
 
 	// Grid snap to pixels even when grid snapping is disabled, if the first component is axis aligned
 	const bool bIsAxisAligned = FMath::IsNearlyZero(FMath::Abs(FMath::Fmod(FirstComponent->GetRotation(), 90.f)));
-	if (!PixelMapping->bGridSnappingEnabled && !bIsAxisAligned)
+	if (!PixelMapping->bGridSnappingEnabled)
 	{
 		return DesiredPosition;
 	}
 
 	const FVector2D CellSize = [PixelMapping, RendererOfFirstComponent]()
 		{
-			if (PixelMapping->bGridSnappingEnabled)
-			{
-				const FVector2D TextureSize = RendererOfFirstComponent->GetSize();
-				return TextureSize / FVector2D(PixelMapping->SnapGridColumns, PixelMapping->SnapGridRows);
-			}
-			
-			// Grid snap to pixels if grid snapping is disabled
-			return FVector2D(1.f, 1.f);
+			const FVector2D TextureSize = RendererOfFirstComponent->GetSize();
+			return TextureSize / FVector2D(PixelMapping->SnapGridColumns, PixelMapping->SnapGridRows);
 		}();
 
 	const int32 Column = FMath::RoundHalfToZero(DesiredPosition.X / CellSize.X);
diff --git a/Engine/Plugins/VirtualProduction/DMX/DMXPixelMapping/Source/DMXPixelMappingEditor/Private/Settings/DMXPixelMappingEditorSettings.h b/Engine/Plugins/VirtualProduction/DMX/DMXPixelMapping/Source/DMXPixelMappingEditor/Private/Settings/DMXPixelMappingEditorSettings.h
--- a/Engine/Plugins/VirtualProduction/DMX/DMXPixelMapping/Source/DMXPixelMappingEditor/Private/Settings/DMXPixelMappingEditorSettings.h
+++ b/Engine/Plugins/VirtualProduction/DMX/DMXPixelMapping/Source/DMXPixelMappingEditor/Private/Settings/DMXPixelMappingEditorSettings.h
@@ -50,6 +50,10 @@ struct FDMXPixelMappingDesignerSettings
 	UPROPERTY()
 	bool bApplyLayoutScriptWhenLoaded = true;
 
+	/** If true, a pivot is displayed for selected components */
+	UPROPERTY()
+	bool bShowPivot = true;
+
 	/**  If true, shows a widget for each cell. It is recommended that this is turned off when pixel mapping large quantities of fixtures. */
 	UPROPERTY()
 	bool bShowMatrixCells = true;
diff --git a/Engine/Plugins/VirtualProduction/DMX/DMXPixelMapping/Source/DMXPixelMappingEditor/Private/Toolkits/DMXPixelMappingToolkit.cpp b/Engine/Plugins/VirtualProduction/DMX/DMXPixelMapping/Source/DMXPixelMappingEditor/Private/Toolkits/DMXPixelMappingToolkit.cpp
--- a/Engine/Plugins/VirtualProduction/DMX/DMXPixelMapping/Source/DMXPixelMappingEditor/Private/Toolkits/DMXPixelMappingToolkit.cpp
+++ b/Engine/Plugins/VirtualProduction/DMX/DMXPixelMapping/Source/DMXPixelMappingEditor/Private/Toolkits/DMXPixelMappingToolkit.cpp
@@ -525,10 +525,7 @@ bool FDMXPixelMappingToolkit::CanSizeSelectedComponentToTexture() const
 			return
 				Component->GetClass() == UDMXPixelMappingFixtureGroupComponent::StaticClass() ||
 				Component->GetClass() == UDMXPixelMappingMatrixComponent::StaticClass() ||
-				Component->GetClass() == UDMXPixelMappingScreenComponent::StaticClass() ||
-				(Component->GetParent() && Component->GetParent()->GetClass() == UDMXPixelMappingFixtureGroupComponent::StaticClass()) ||
-				(Component->GetParent() && Component->GetParent()->GetClass() == UDMXPixelMappingMatrixComponent::StaticClass()) ||
-				(Component->GetParent() && Component->GetParent()->GetClass() == UDMXPixelMappingScreenComponent::StaticClass());
+				Component->GetClass() == UDMXPixelMappingScreenComponent::StaticClass();
 		}
 	}
 	return false;
@@ -878,14 +875,14 @@ void FDMXPixelMappingToolkit::SetupCommands()
 	);
 
 	GetToolkitCommands()->MapAction(
-		FDMXPixelMappingEditorCommands::Get().ResizeMode,
+		FDMXPixelMappingEditorCommands::Get().EnableResizeMode,
 		FExecuteAction::CreateSP(this, &FDMXPixelMappingToolkit::SetTransformHandleMode, EDMXPixelMappingTransformHandleMode::Resize),
 		FCanExecuteAction(),
 		FGetActionCheckState::CreateSP(this, &FDMXPixelMappingToolkit::GetTransformHandleModeCheckboxState, EDMXPixelMappingTransformHandleMode::Resize)
 	);
 
 	GetToolkitCommands()->MapAction(
-		FDMXPixelMappingEditorCommands::Get().RotateMode,
+		FDMXPixelMappingEditorCommands::Get().EnableRotateMode,
 		FExecuteAction::CreateSP(this, &FDMXPixelMappingToolkit::SetTransformHandleMode, EDMXPixelMappingTransformHandleMode::Rotate),
 		FCanExecuteAction(),
 		FGetActionCheckState::CreateSP(this, &FDMXPixelMappingToolkit::GetTransformHandleModeCheckboxState, EDMXPixelMappingTransformHandleMode::Rotate)
@@ -897,10 +894,10 @@ void FDMXPixelMappingToolkit::SetupCommands()
 	);
 
 	// Designer related
-	constexpr bool bTransactSizeComponentToTexture = true;
+	constexpr bool bSizeComponentToTexture = true;
 	GetToolkitCommands()->MapAction(
 		FDMXPixelMappingEditorCommands::Get().SizeComponentToTexture,
-		FExecuteAction::CreateSP(this, &FDMXPixelMappingToolkit::SizeSelectedComponentToTexture, bTransactSizeComponentToTexture),
+		FExecuteAction::CreateSP(this, &FDMXPixelMappingToolkit::SizeSelectedComponentToTexture, bSizeComponentToTexture),
 		FCanExecuteAction::CreateSP(this, &FDMXPixelMappingToolkit::CanSizeSelectedComponentToTexture)
 	);
 
@@ -912,6 +909,7 @@ void FDMXPixelMappingToolkit::SetupCommands()
 	UE_DMX_MAP_EDITOR_SETTING_TO_TOGGLE_COMMAND(EditorSettings, DesignerSettings, bShowComponentNames, ToggleShowComponentNames);
 	UE_DMX_MAP_EDITOR_SETTING_TO_TOGGLE_COMMAND(EditorSettings, DesignerSettings, bShowPatchInfo, ToggleShowPatchInfo);
 	UE_DMX_MAP_EDITOR_SETTING_TO_TOGGLE_COMMAND(EditorSettings, DesignerSettings, bShowCellIDs, ToggleShowCellIDs);
+	UE_DMX_MAP_EDITOR_SETTING_TO_TOGGLE_COMMAND(EditorSettings, DesignerSettings, bShowPivot, ToggleShowPivot);
 }
 #undef UE_DMX_MAP_EDITOR_SETTING_TO_TOGGLE_COMMAND
 
diff --git a/Engine/Plugins/VirtualProduction/DMX/DMXPixelMapping/Source/DMXPixelMappingEditor/Private/ViewModels/DMXPixelMappingOutputComponentModel.cpp b/Engine/Plugins/VirtualProduction/DMX/DMXPixelMapping/Source/DMXPixelMappingEditor/Private/ViewModels/DMXPixelMappingOutputComponentModel.cpp
--- a/Engine/Plugins/VirtualProduction/DMX/DMXPixelMapping/Source/DMXPixelMappingEditor/Private/ViewModels/DMXPixelMappingOutputComponentModel.cpp
+++ b/Engine/Plugins/VirtualProduction/DMX/DMXPixelMapping/Source/DMXPixelMappingEditor/Private/ViewModels/DMXPixelMappingOutputComponentModel.cpp
@@ -104,20 +104,18 @@ namespace UE::DMX
 	{
 		const FDMXPixelMappingDesignerSettings& DesignerSettings = GetDefault<UDMXPixelMappingEditorSettings>()->DesignerSettings;
 
-		return DesignerSettings.bShowComponentNames &&
+		return 
+			DesignerSettings.bShowComponentNames &&
 			OutputComponent &&
 			OutputComponent->GetClass() != UDMXPixelMappingMatrixCellComponent::StaticClass();
 	}
 
 	bool FDMXPixelMappingOutputComponentModel::ShouldDrawNameAbove() const
 	{
-		if (OutputComponent)
-		{
-			return
-				OutputComponent->GetClass() == UDMXPixelMappingFixtureGroupComponent::StaticClass() ||
-				OutputComponent->GetClass() == UDMXPixelMappingScreenComponent::StaticClass();
-		}
-		return true;
+		return
+			OutputComponent &&
+			(OutputComponent->GetClass() == UDMXPixelMappingFixtureGroupComponent::StaticClass() ||
+			OutputComponent->GetClass() == UDMXPixelMappingScreenComponent::StaticClass());
 	}
 
 	bool FDMXPixelMappingOutputComponentModel::ShouldDrawCellID() const
@@ -136,7 +134,16 @@ namespace UE::DMX
 
 		return
 			DesignerSettings.bShowPatchInfo &&
-			GetFixturePatch() != nullptr;;
+			GetFixturePatch() != nullptr;
+	}
+
+	bool FDMXPixelMappingOutputComponentModel::ShouldDrawPivot() const
+	{
+		const FDMXPixelMappingDesignerSettings& DesignerSettings = GetDefault<UDMXPixelMappingEditorSettings>()->DesignerSettings;
+
+		return 
+			bSelected &&
+			DesignerSettings.bShowPivot;
 	}
 
 	FText FDMXPixelMappingOutputComponentModel::GetAddressesText() const
@@ -184,7 +191,7 @@ namespace UE::DMX
 				}
 				else
 				{
-					const float Opacity = OutputComponent->GetEditorColor().A * .6f;
+					const float Opacity = OutputComponent->GetEditorColor().A * .5f;
 					return OutputComponent->GetEditorColor().CopyWithNewOpacity(Opacity);
 				}
 			}
diff --git a/Engine/Plugins/VirtualProduction/DMX/DMXPixelMapping/Source/DMXPixelMappingEditor/Private/ViewModels/DMXPixelMappingOutputComponentModel.h b/Engine/Plugins/VirtualProduction/DMX/DMXPixelMapping/Source/DMXPixelMappingEditor/Private/ViewModels/DMXPixelMappingOutputComponentModel.h
--- a/Engine/Plugins/VirtualProduction/DMX/DMXPixelMapping/Source/DMXPixelMappingEditor/Private/ViewModels/DMXPixelMappingOutputComponentModel.h
+++ b/Engine/Plugins/VirtualProduction/DMX/DMXPixelMapping/Source/DMXPixelMappingEditor/Private/ViewModels/DMXPixelMappingOutputComponentModel.h
@@ -58,9 +58,12 @@ namespace UE::DMX
 		/** Returns true if the cell ID should be shown */
 		bool ShouldDrawCellID() const;
 
-		/** Returns true if the patch inifo should be shown */
+		/** Returns true if the patch info should be shown */
 		bool ShouldDrawPatchInfo() const;
 
+		/** Returns true if the pivot should be shown */
+		bool ShouldDrawPivot() const;
+
 		/** Returns the addresses of the patch as text */
 		FText GetAddressesText() const;
 
diff --git a/Engine/Plugins/VirtualProduction/DMX/DMXPixelMapping/Source/DMXPixelMappingEditor/Private/Views/SDMXPixelMappingDesignerView.cpp b/Engine/Plugins/VirtualProduction/DMX/DMXPixelMapping/Source/DMXPixelMappingEditor/Private/Views/SDMXPixelMappingDesignerView.cpp
--- a/Engine/Plugins/VirtualProduction/DMX/DMXPixelMapping/Source/DMXPixelMappingEditor/Private/Views/SDMXPixelMappingDesignerView.cpp
+++ b/Engine/Plugins/VirtualProduction/DMX/DMXPixelMapping/Source/DMXPixelMappingEditor/Private/Views/SDMXPixelMappingDesignerView.cpp
@@ -4,6 +4,7 @@
 
 #include "Algo/Find.h"
 #include "DMXPixelMapping.h"
+#include "DMXPixelMappingEditorCommands.h"
 #include "DMXPixelMappingEditorStyle.h"
 #include "DragDrop/DMXPixelMappingGroupChildDragDropHelper.h"
 #include "Framework/Application/SlateApplication.h"
@@ -432,30 +433,36 @@ FReply SDMXPixelMappingDesignerView::OnMouseButtonDown(const FGeometry& MyGeomet
 
 FReply SDMXPixelMappingDesignerView::OnMouseButtonUp(const FGeometry& MyGeometry, const FPointerEvent& MouseEvent)
 {
-	SDMXPixelMappingSurface::OnMouseButtonUp(MyGeometry, MouseEvent);
-
-	// Select the output component under the mouse
-	PendingSelectedComponent = GetComponentUnderCursor();
-
-	// Select the Renderer if no Output Component was under the mouse
-	if (!PendingSelectedComponent.IsValid())
+	if (MouseEvent.GetEffectingButton() == EKeys::LeftMouseButton)
 	{
-		if (TSharedPtr<FDMXPixelMappingToolkit> ToolkitPtr = WeakToolkit.Pin())
+		// Select the output component under the mouse
+		PendingSelectedComponent = GetComponentUnderCursor();
+
+		// Select the Renderer if no Output Component was under the mouse
+		if (!PendingSelectedComponent.IsValid())
 		{
-			if (UDMXPixelMappingRendererComponent* RendererComponent = ToolkitPtr->GetActiveRendererComponent())
+			if (TSharedPtr<FDMXPixelMappingToolkit> ToolkitPtr = WeakToolkit.Pin())
 			{
-				PendingSelectedComponent = RendererComponent;
+				if (UDMXPixelMappingRendererComponent* RendererComponent = ToolkitPtr->GetActiveRendererComponent())
+				{
+					PendingSelectedComponent = RendererComponent;
+				}
 			}
 		}
-	}
 
-	const bool bClearPreviousSelection = !MouseEvent.IsShiftDown() && !MouseEvent.IsControlDown();
-	ResolvePendingSelectedComponents(bClearPreviousSelection);
+		const bool bClearPreviousSelection = !MouseEvent.IsShiftDown() && !MouseEvent.IsControlDown();
+		ResolvePendingSelectedComponents(bClearPreviousSelection);
+	}
+	else if (MouseEvent.GetEffectingButton() == EKeys::RightMouseButton &&
+		!bIsPanning &&
+		!bIsZooming)
+	{
+		OpenContextMenu();
+	}
+	
+	SDMXPixelMappingSurface::OnMouseButtonUp(MyGeometry, MouseEvent);
 
-	return FReply::Handled()
-		.EndDragDrop()
-		.ReleaseMouseCapture()
-		.SetUserFocus(AsShared());
+	return FReply::Handled().ReleaseMouseCapture();
 }
 
 FReply SDMXPixelMappingDesignerView::OnKeyDown(const FGeometry& MyGeometry, const FKeyEvent& InKeyEvent)
@@ -742,6 +749,66 @@ void SDMXPixelMappingDesignerView::ClearExtensionWidgets()
 	ExtensionWidgetCanvas->ClearChildren();
 }
 
+void SDMXPixelMappingDesignerView::OpenContextMenu()
+{
+	const TSharedPtr<FDMXPixelMappingToolkit> Toolkit = WeakToolkit.Pin();
+	if (!Toolkit.IsValid())
+	{
+		return;
+	}
+
+	constexpr bool bShouldCloseMenuAfterSelection = true;
+	FMenuBuilder MenuBuilder(bShouldCloseMenuAfterSelection, Toolkit->GetToolkitCommands());
+
+	MenuBuilder.BeginSection("EditorSettings", LOCTEXT("EditorSettingsSection", "Editor Settings"));
+	{
+		MenuBuilder.AddMenuEntry(FDMXPixelMappingEditorCommands::Get().ToggleAlwaysSelectGroup);
+		MenuBuilder.AddMenuEntry(FDMXPixelMappingEditorCommands::Get().ToggleScaleChildrenWithParent);
+	}
+	MenuBuilder.EndSection();
+
+	MenuBuilder.BeginSection("TransformModes", LOCTEXT("TransformModesSection", "Transform Modes"));
+	{		
+		MenuBuilder.AddMenuEntry(FDMXPixelMappingEditorCommands::Get().EnableResizeMode);
+		MenuBuilder.AddMenuEntry(FDMXPixelMappingEditorCommands::Get().EnableRotateMode);
+	}
+	MenuBuilder.EndSection();
+
+	MenuBuilder.BeginSection("DisplaySettings", LOCTEXT("SettingsSection", "Display Settings"));
+	{
+		MenuBuilder.AddMenuEntry(FDMXPixelMappingEditorCommands::Get().ToggleShowComponentNames);
+		MenuBuilder.AddMenuEntry(FDMXPixelMappingEditorCommands::Get().ToggleShowPatchInfo);
+		MenuBuilder.AddMenuEntry(FDMXPixelMappingEditorCommands::Get().ToggleShowMatrixCells);
+		MenuBuilder.AddMenuEntry(FDMXPixelMappingEditorCommands::Get().ToggleShowCellIDs);
+		MenuBuilder.AddMenuEntry(FDMXPixelMappingEditorCommands::Get().ToggleShowPivot);
+	}
+	MenuBuilder.EndSection();
+
+	if (Toolkit->CanSizeSelectedComponentToTexture())
+	{
+		MenuBuilder.BeginSection("AdvancedActions", LOCTEXT("ActionsSection", "Advanced Actions"));
+		{
+			MenuBuilder.AddMenuEntry(FDMXPixelMappingEditorCommands::Get().SizeComponentToTexture);
+		}
+		MenuBuilder.EndSection();
+	}
+
+	const TSharedRef<SWidget> MenuContent = 
+		SNew(SBorder)
+		.BorderImage(FAppStyle::GetBrush("Menu.Background"))
+		[
+			MenuBuilder.MakeWidget()
+		];
+
+	FSlateApplication::Get().PushMenu(
+		AsShared(),
+		FWidgetPath(),
+		MenuContent,
+		FSlateApplication::Get().GetCursorPos(),
+		FPopupTransitionEffect(FPopupTransitionEffect::ContextMenu)
+	);
+}
+
 EVisibility SDMXPixelMappingDesignerView::GetExtensionCanvasVisibility() const
 {
 	for (const FDMXPixelMappingComponentReference& Component : GetSelectedComponents())
@@ -786,18 +853,18 @@ FVector2D SDMXPixelMappingDesignerView::GetExtensionPosition(TSharedPtr<SDMXPixe
 			}();
 
 
-			FVector2D SelectedWidgetScale = FVector2D(SelectedComponentGeometry.GetAccumulatedRenderTransform().GetMatrix().GetScale().GetVector());
+			const FVector2D SelectedWidgetScale = FVector2D(SelectedComponentGeometry.GetAccumulatedRenderTransform().GetMatrix().GetScale().GetVector());
 
-			FVector2D ApplicationScaledOffset = Handle->GetOffset() * GetDesignerGeometry().Scale;
+			const FVector2D ApplicationScaledOffset = Handle->GetOffset() * GetDesignerGeometry().Scale;
 
-			FVector2D LocalOffsetFull = ApplicationScaledOffset / SelectedWidgetScale;
-			FVector2D PositionFullOffset = GetDesignerGeometry().AbsoluteToLocal(SelectedComponentGeometry.LocalToAbsolute(LocalPosition + LocalOffsetFull));
-			FVector2D LocalOffsetHalf = (ApplicationScaledOffset / 2.0f) / SelectedWidgetScale;
-			FVector2D PositionHalfOffset = GetDesignerGeometry().AbsoluteToLocal(SelectedComponentGeometry.LocalToAbsolute(LocalPosition + LocalOffsetHalf));
+			const FVector2D LocalOffsetFull = ApplicationScaledOffset / SelectedWidgetScale;
+			const FVector2D PositionFullOffset = GetDesignerGeometry().AbsoluteToLocal(SelectedComponentGeometry.LocalToAbsolute(LocalPosition + LocalOffsetFull));
+			const FVector2D LocalOffsetHalf = (ApplicationScaledOffset / 2.0f) / SelectedWidgetScale;
+			const FVector2D PositionHalfOffset = GetDesignerGeometry().AbsoluteToLocal(SelectedComponentGeometry.LocalToAbsolute(LocalPosition + LocalOffsetHalf));
 
-			FVector2D PivotCorrection = PositionHalfOffset - (PositionFullOffset + FVector2D(5.0f, 5.0f));
+			const FVector2D PivotCorrection = PositionHalfOffset - (PositionFullOffset + FVector2D(5.0f, 5.0f));
 
-			FVector2D FinalPosition = PositionFullOffset + PivotCorrection;
+			const FVector2D FinalPosition = PositionFullOffset + PivotCorrection;
 
 			return FVector2D(FinalPosition);
 		}
diff --git a/Engine/Plugins/VirtualProduction/DMX/DMXPixelMapping/Source/DMXPixelMappingEditor/Private/Views/SDMXPixelMappingDesignerView.h b/Engine/Plugins/VirtualProduction/DMX/DMXPixelMapping/Source/DMXPixelMappingEditor/Private/Views/SDMXPixelMappingDesignerView.h
--- a/Engine/Plugins/VirtualProduction/DMX/DMXPixelMapping/Source/DMXPixelMappingEditor/Private/Views/SDMXPixelMappingDesignerView.h
+++ b/Engine/Plugins/VirtualProduction/DMX/DMXPixelMapping/Source/DMXPixelMappingEditor/Private/Views/SDMXPixelMappingDesignerView.h
@@ -96,6 +96,9 @@ private:
 	/** Clears the extension widgets */
 	void ClearExtensionWidgets();
 
+	/** Opens the designer context menu */
+	void OpenContextMenu();
+
 	/** Returns if the extension widget canvas should be visible */
 	EVisibility GetExtensionCanvasVisibility() const;
 
diff --git a/Engine/Plugins/VirtualProduction/DMX/DMXPixelMapping/Source/DMXPixelMappingEditor/Private/Widgets/SDMXPixelMappingDesignerToolbar.cpp b/Engine/Plugins/VirtualProduction/DMX/DMXPixelMapping/Source/DMXPixelMappingEditor/Private/Widgets/SDMXPixelMappingDesignerToolbar.cpp
--- a/Engine/Plugins/VirtualProduction/DMX/DMXPixelMapping/Source/DMXPixelMappingEditor/Private/Widgets/SDMXPixelMappingDesignerToolbar.cpp
+++ b/Engine/Plugins/VirtualProduction/DMX/DMXPixelMapping/Source/DMXPixelMappingEditor/Private/Widgets/SDMXPixelMappingDesignerToolbar.cpp
@@ -4,13 +4,14 @@
 
 #include "DMXPixelMapping.h"
 #include "DMXPixelMappingEditorCommands.h"
+#include "SEditorViewportToolBarMenu.h"
 #include "Styling/AppStyle.h"
 #include "Styling/ToolBarStyle.h"
 #include "SViewportToolBarComboMenu.h"
-#include "ToolMenus.h"
 #include "Toolkits/DMXPixelMappingToolkit.h"
 #include "Widgets/Images/SImage.h"
 #include "Widgets/Input/SButton.h"
+#include "Widgets/Input/SNumericEntryBox.h"
 #include "Widgets/Layout/SBorder.h"
 #include "Widgets/Text/STextBlock.h"
 #include "Widgets/SDMXPixelMappingSnapGridMenu.h"
@@ -18,174 +19,163 @@
 
 #define LOCTEXT_NAMESPACE "SDMXPixelMappingDesignerToolbar"
 
-
 namespace UE::DMX
 {
-	namespace DMXPixelMappingDesignerToolbar::Private
-	{
-		constexpr TCHAR ToolbarName[] = TEXT("PixelMapping.DesignerToolbar");
-	}
-
-	SDMXPixelMappingDesignerToolbar::~SDMXPixelMappingDesignerToolbar()
-	{
-		using namespace DMXPixelMappingDesignerToolbar::Private;
-		if (UToolMenus::Get()->IsMenuRegistered(ToolbarName))
-		{
-			UToolMenus::Get()->RemoveMenu(ToolbarName);
-		}
-	}
-
 	void SDMXPixelMappingDesignerToolbar::Construct(const FArguments& InArgs, const TSharedRef<FDMXPixelMappingToolkit>& InToolkit)
 	{
 		WeakToolkit = InToolkit;
 
-		RegisterToolbarMenu(InArgs);
-		
-		const FToolMenuContext Context(InToolkit->GetToolkitCommands());
+		constexpr TCHAR ToolBarStyleName[] = TEXT("EditorViewportToolBar");
+		const FToolBarStyle& ToolBarStyle = FAppStyle::Get().GetWidgetStyle<FToolBarStyle>(ToolBarStyleName);
+
+		const TSharedRef<FUICommandList> CommandList = InToolkit->GetToolkitCommands();
 
-		using namespace DMXPixelMappingDesignerToolbar::Private;
 		ChildSlot
 		[
-			UToolMenus::Get()->GenerateWidget(ToolbarName, Context)
-		];
-	}
-
-	void SDMXPixelMappingDesignerToolbar::RegisterToolbarMenu(const FArguments& InArgs)
-	{
-		using namespace UE::DMX;
-
-		UToolMenus* ToolMenus = UToolMenus::Get();
-
-		using namespace DMXPixelMappingDesignerToolbar::Private;
-		if (ToolMenus->IsMenuRegistered(ToolbarName))
-		{
-			return;
-		}
-
-		UToolMenu* Toolbar = UToolMenus::Get()->RegisterMenu(ToolbarName, NAME_None, EMultiBoxType::SlimHorizontalToolBar);
-		Toolbar->StyleName = "EditorViewportToolbar";
-		
-
-		// Transform Handle Modes
-		{
-			FToolMenuSection& Section = Toolbar->AddSection("TransformHandleMode");
-
-			// Resize mode
-			const FCheckBoxStyle& CheckBoxStartStyle = FAppStyle::Get().GetWidgetStyle<FCheckBoxStyle>("EditorViewportToolBar.ToggleButton.Start");
-
-			TSharedRef<SCheckBox> EnableResizeModeToggleButton = 
-				SNew(SCheckBox)
-				.Style(&CheckBoxStartStyle)
-				.ToolTipText(LOCTEXT("TransformHandleResizeMode", "Resize Components"))
-				.OnCheckStateChanged(this, &SDMXPixelMappingDesignerToolbar::OnTransformHandleModeSelected, EDMXPixelMappingTransformHandleMode::Resize)
-				.IsChecked(this, &SDMXPixelMappingDesignerToolbar::GetCheckboxStateForTransormHandleMode, EDMXPixelMappingTransformHandleMode::Resize)
-				[
-					SNew(SImage)
-					.Image(FAppStyle::GetBrush("EditorViewport.ScaleMode"))
-					.ColorAndOpacity(FSlateColor::UseForeground())
-				];
-
-			// Rotate mode
-			const FCheckBoxStyle& CheckBoxEndStyle = FAppStyle::Get().GetWidgetStyle<FCheckBoxStyle>("EditorViewportToolBar.ToggleButton.End");
+			SNew(SBorder)
+			.BorderImage(FAppStyle::GetBrush("NoBorder"))
+			.ForegroundColor(FAppStyle::GetSlateColor("DefaultForeground"))
+			[
+				SNew(SHorizontalBox)
 
-			TSharedRef<SCheckBox> EnableRotateModeToggleButton = 
-				SNew(SCheckBox)
-				.Style(&CheckBoxEndStyle)
-				.ToolTipText(LOCTEXT("TransformHandleRotateMode", "Rotate Components"))
-				.OnCheckStateChanged(this, &SDMXPixelMappingDesignerToolbar::OnTransformHandleModeSelected, EDMXPixelMappingTransformHandleMode::Rotate)
-				.IsChecked(this, &SDMXPixelMappingDesignerToolbar::GetCheckboxStateForTransormHandleMode, EDMXPixelMappingTransformHandleMode::Rotate)
+				+ SHorizontalBox::Slot()
+				.AutoWidth()
+				.VAlign(VAlign_Center)
+				.Padding(4.f, 2.0f)
 				[
-					SNew(SImage)
-					.Image(FAppStyle::GetBrush("EditorViewport.RotateMode"))
-					.ColorAndOpacity(FSlateColor::UseForeground())
-				];
+					GenerateTransformHandleModeSection()
+				]
 
-			// As a single widget
-			const TSharedRef<SWidget> TransformHandleModeWidget =
-				SNew(SHorizontalBox)
 				+ SHorizontalBox::Slot()
-				.Padding(0.f)
+				.AutoWidth()
+				.VAlign(VAlign_Center)
+				.Padding(4.f, 2.0f)
 				[
-					EnableResizeModeToggleButton
+					GenerateGridSnappingSection()
 				]
+					
 				+ SHorizontalBox::Slot()
-				.Padding(0.f)
+				.AutoWidth()
+				.VAlign(VAlign_Center)
+				.Padding(4.f, 2.0f)
 				[
-					EnableRotateModeToggleButton
-				];
+					GenerateZoomToFitSection(InArgs)
+				]
 
-			Section.AddEntry
-			(
-				FToolMenuEntry::InitWidget
-				(
-					"TransformHandleModes",
-					TransformHandleModeWidget,
-					FText::GetEmpty()
-				)
-			);
-		}
+				+ SHorizontalBox::Slot()
+				.AutoWidth()
+				.VAlign(VAlign_Center)
+				.Padding(4.f, 2.0f)
+				[
+					GenerateSettingsSection()
+				]
+			]
+		];
+	}
 
+	TSharedRef<SWidget> SDMXPixelMappingDesignerToolbar::GenerateTransformHandleModeSection()
+	{
+		// Resize mode
+		const FCheckBoxStyle& CheckBoxStartStyle = FAppStyle::Get().GetWidgetStyle<FCheckBoxStyle>("EditorViewportToolBar.ToggleButton.Start");
+
+		TSharedRef<SCheckBox> EnableResizeModeToggleButton = 
+			SNew(SCheckBox)
+			.Style(&CheckBoxStartStyle)
+			.ToolTipText(LOCTEXT("TransformHandleResizeMode", "Resize Components"))
+			.OnCheckStateChanged(this, &SDMXPixelMappingDesignerToolbar::OnTransformHandleModeSelected, EDMXPixelMappingTransformHandleMode::Resize)
+			.IsChecked(this, &SDMXPixelMappingDesignerToolbar::GetCheckboxStateForTransormHandleMode, EDMXPixelMappingTransformHandleMode::Resize)
+			[
+				SNew(SImage)
+				.Image(FAppStyle::GetBrush("EditorViewport.ScaleMode"))
+				.ColorAndOpacity(FSlateColor::UseForeground())
+			];
+
+		// Rotate mode
+		const FCheckBoxStyle& CheckBoxEndStyle = FAppStyle::Get().GetWidgetStyle<FCheckBoxStyle>("EditorViewportToolBar.ToggleButton.End");
+
+		TSharedRef<SCheckBox> EnableRotateModeToggleButton = 
+			SNew(SCheckBox)
+			.Style(&CheckBoxEndStyle)
+			.ToolTipText(LOCTEXT("TransformHandleRotateMode", "Rotate Components"))
+			.OnCheckStateChanged(this, &SDMXPixelMappingDesignerToolbar::OnTransformHandleModeSelected, EDMXPixelMappingTransformHandleMode::Rotate)
+			.IsChecked(this, &SDMXPixelMappingDesignerToolbar::GetCheckboxStateForTransormHandleMode, EDMXPixelMappingTransformHandleMode::Rotate)
+			[
+				SNew(SImage)
+				.Image(FAppStyle::GetBrush("EditorViewport.RotateMode"))
+				.ColorAndOpacity(FSlateColor::UseForeground())
+			];
+
+		// As a single widget
+		const TSharedRef<SWidget> TransformHandleModeWidget =
+			SNew(SHorizontalBox)
+			+ SHorizontalBox::Slot()
+			.Padding(0.f)
+			[
+				EnableResizeModeToggleButton
+			]
+			+ SHorizontalBox::Slot()
+			.Padding(0.f)
+			[
+				EnableRotateModeToggleButton
+			];
+
+		return TransformHandleModeWidget;
+	}
 
-		// Grid Snapping
-		{
-			FToolMenuSection& Section = Toolbar->AddSection("GridSnapping");
-
-			const FUICommandInfo& ToggleSnapGridCommand = *FDMXPixelMappingEditorCommands::Get().ToggleGridSnapping;
-
-			const TSharedRef<SViewportToolBarComboMenu> GridSnappingComboMenu =
-				SNew(SViewportToolBarComboMenu)
-				.Cursor(EMouseCursor::Default)
-				.IsChecked(this, &SDMXPixelMappingDesignerToolbar::GetSnapGridEnabledCheckState)
-				.OnCheckStateChanged(this, &SDMXPixelMappingDesignerToolbar::OnSnapGridCheckStateChanged)
-				.Label(this, &SDMXPixelMappingDesignerToolbar::GetSnapGridLabel)
-				.OnGetMenuContent(this, &SDMXPixelMappingDesignerToolbar::GenerateSnapGridMenu)
-				.ToggleButtonToolTip(ToggleSnapGridCommand.GetDescription())
-				.MenuButtonToolTip(LOCTEXT("SnapGridMenuTooltip", "Grid Snapping Settings"))
-				.Icon(FSlateIcon(FAppStyle::GetAppStyleSetName(), "EditorViewport.LocationGridSnap"))
-				.ParentToolBar(SharedThis(this));
-
-			Section.AddEntry
-			(
-				FToolMenuEntry::InitWidget
-				(
-					"GridSnappingComboButton",
-					GridSnappingComboMenu,
-					FText::GetEmpty()
-				)
-			);
-		}
+	TSharedRef<SWidget> SDMXPixelMappingDesignerToolbar::GenerateGridSnappingSection()
+	{
+		const FUICommandInfo* ToggleSnapGridCommand = FDMXPixelMappingEditorCommands::Get().ToggleGridSnapping.Get();
+		check(ToggleSnapGridCommand);
 
+		const TSharedRef<SViewportToolBarComboMenu> GridSnappingComboMenu =
+			SNew(SViewportToolBarComboMenu)
+			.Cursor(EMouseCursor::Default)
+			.IsChecked(this, &SDMXPixelMappingDesignerToolbar::GetSnapGridEnabledCheckState)
+			.OnCheckStateChanged(this, &SDMXPixelMappingDesignerToolbar::OnSnapGridCheckStateChanged)
+			.Label(this, &SDMXPixelMappingDesignerToolbar::GetSnapGridLabel)
+			.OnGetMenuContent(this, &SDMXPixelMappingDesignerToolbar::GenerateSnapGridMenu)
+			.ToggleButtonToolTip(ToggleSnapGridCommand->GetDescription())
+			.MenuButtonToolTip(LOCTEXT("SnapGridMenuTooltip", "Grid Snapping Settings"))
+			.Icon(FSlateIcon(FAppStyle::GetAppStyleSetName(), "EditorViewport.LocationGridSnap"))
+			.ParentToolBar(SharedThis(this));
+
+		return GridSnappingComboMenu;
+	}
 
-		// Zoom to Fit
-		{
-			FToolMenuSection& Section = Toolbar->AddSection("ZoomToFit");
+	TSharedRef<SWidget> SDMXPixelMappingDesignerToolbar::GenerateZoomToFitSection(const FArguments& InArgs)
+	{
+		constexpr TCHAR ToolBarStyleName[] = TEXT("EditorViewportToolBar");
+		const FToolBarStyle& ToolBarStyle = FAppStyle::Get().GetWidgetStyle<FToolBarStyle>(ToolBarStyleName);
+
+		const TSharedRef<SButton> ZoomToFitButton = 
+			SNew(SButton)
+			.ButtonStyle(&ToolBarStyle.ButtonStyle)
+			.ToolTipText(LOCTEXT("ZoomToFitToolTip", "Zoom To Fit"))
+			.OnClicked(InArgs._OnZoomToFitClicked)
+			.ContentPadding(ToolBarStyle.ButtonPadding)
+			.VAlign(VAlign_Center)
+			[
+				SNew(SImage)
+				.Image(FAppStyle::GetBrush("UMGEditor.ZoomToFit"))
+				.ColorAndOpacity(FSlateColor::UseForeground())
+			];
+
+		return ZoomToFitButton;
+	}
 
-			constexpr TCHAR ToolBarStyleName[] = TEXT("EditorViewportToolBar");
-			const FToolBarStyle& ToolBarStyle = FAppStyle::Get().GetWidgetStyle<FToolBarStyle>(ToolBarStyleName);
+	TSharedRef<SWidget> SDMXPixelMappingDesignerToolbar::GenerateSettingsSection()
+	{
+		constexpr TCHAR ToolBarStyleName[] = TEXT("EditorViewportToolBar");
+		const FToolBarStyle& ToolBarStyle = FAppStyle::Get().GetWidgetStyle<FToolBarStyle>(ToolBarStyleName);
 
-			const TSharedRef<SButton> ZoomToFitButton = 
-				SNew(SButton)
-				.ButtonStyle(&ToolBarStyle.ButtonStyle)
-				.ToolTipText(LOCTEXT("ZoomToFit_ToolTip", "Zoom To Fit"))
-				.OnClicked(InArgs._OnZoomToFitClicked)
-				.ContentPadding(ToolBarStyle.ButtonPadding)
-				.VAlign(VAlign_Center)
-				[
-					SNew(SImage)
-					.Image(FAppStyle::GetBrush("UMGEditor.ZoomToFit"))
-					.ColorAndOpacity(FSlateColor::UseForeground())
-				];
+		const TSharedRef<SEditorViewportToolbarMenu> SettingsButton =
+			SNew(SEditorViewportToolbarMenu)
+			.ParentToolBar(SharedThis(this))
+			.Label(LOCTEXT("SettingsButtonLabel", "Settings"))
+			.LabelIcon(FAppStyle::Get().GetBrush("Icons.Settings"))
+			.OnGetMenuContent(this, &SDMXPixelMappingDesignerToolbar::GenerateSettingsMenuContent);
 
-			Section.AddEntry
-			(
-				FToolMenuEntry::InitWidget
-				(
-					"ZoomToFitButton",
-					ZoomToFitButton,
-					FText::GetEmpty()
-				)
-			);
-		}
+		return SettingsButton;
 	}
 
 	TSharedRef<SWidget> SDMXPixelMappingDesignerToolbar::GenerateSnapGridMenu()
@@ -262,6 +252,93 @@ namespace UE::DMX
 
 		return Toolkit->GetTransformHandleMode() == TransformHandleMode ? ECheckBoxState::Checked : ECheckBoxState::Unchecked;
 	}
+
+	TSharedRef<SWidget> SDMXPixelMappingDesignerToolbar::GenerateSettingsMenuContent()
+	{
+		const TSharedPtr<FDMXPixelMappingToolkit> Toolkit = WeakToolkit.Pin();
+		if (!Toolkit.IsValid())
+		{
+			return SNullWidget::NullWidget;
+		}
+
+		constexpr bool bShouldCloseMenuAfterSelection = true;
+		FMenuBuilder MenuBuilder(bShouldCloseMenuAfterSelection, Toolkit->GetToolkitCommands());
+
+		MenuBuilder.BeginSection("EditorSettings", LOCTEXT("EditorSettingsSection", "Editor Settings"));
+		{
+			MenuBuilder.AddMenuEntry(FDMXPixelMappingEditorCommands::Get().ToggleAlwaysSelectGroup);
+			MenuBuilder.AddMenuEntry(FDMXPixelMappingEditorCommands::Get().ToggleScaleChildrenWithParent);
+		}
+		MenuBuilder.EndSection();
+
+		MenuBuilder.BeginSection("DisplaySettings", LOCTEXT("SettingsSection", "Display Settings"));
+		{
+			MenuBuilder.AddMenuEntry(FDMXPixelMappingEditorCommands::Get().ToggleShowComponentNames);
+			MenuBuilder.AddMenuEntry(FDMXPixelMappingEditorCommands::Get().ToggleShowPatchInfo);
+			MenuBuilder.AddMenuEntry(FDMXPixelMappingEditorCommands::Get().ToggleShowMatrixCells);
+			MenuBuilder.AddMenuEntry(FDMXPixelMappingEditorCommands::Get().ToggleShowCellIDs);
+			MenuBuilder.AddMenuEntry(FDMXPixelMappingEditorCommands::Get().ToggleShowPivot);
+
+			constexpr bool NoIndent = true;
+			MenuBuilder.AddWidget(
+				GenerateComponentFontSizeEditWidget(),
+				LOCTEXT("FontSizeLabel", "Font Size"),
+				NoIndent
+			);
+		}
+		MenuBuilder.EndSection();
+
+		MenuBuilder.BeginSection("LayoutSettings", LOCTEXT("LayoutSettingsSection", "Layout Settings"));
+		{
+			MenuBuilder.AddMenuEntry(FDMXPixelMappingEditorCommands::Get().ToggleApplyLayoutScriptWhenLoaded);
+		}
+
+		return MenuBuilder.MakeWidget();
+	}
+
+	TSharedRef<SWidget> SDMXPixelMappingDesignerToolbar::GenerateComponentFontSizeEditWidget()
+	{
+		UDMXPixelMapping* PixelMapping = WeakToolkit.IsValid() ? WeakToolkit.Pin()->GetDMXPixelMapping() : nullptr;
+		if (PixelMapping)
+		{
+			return
+				SNew(SBorder)
+				.BorderImage(FAppStyle::GetBrush("NoBorder"))
+				.HAlign(HAlign_Right)
+				[
+					SNew(SNumericEntryBox<uint8>)
+					.MinDesiredValueWidth(40.f)
+					.MinSliderValue(6)
+					.MaxSliderValue(24)
+					.AllowSpin(true)
+					.Value(this, &SDMXPixelMappingDesignerToolbar::GetComponentFontSize)
+					.OnValueChanged(this, &SDMXPixelMappingDesignerToolbar::SetComponentFontSize)
+				];
+		}
+
+		return SNullWidget::NullWidget;
+	}
+
+	TOptional<uint8> SDMXPixelMappingDesignerToolbar::GetComponentFontSize() const
+	{
+		UDMXPixelMapping* PixelMapping = WeakToolkit.IsValid() ? WeakToolkit.Pin()->GetDMXPixelMapping() : nullptr;
+		if (PixelMapping)
+		{
+			return PixelMapping->ComponentLabelFontSize;
+		}
+
+		return TOptional<uint8>();
+	}
+
+	void SDMXPixelMappingDesignerToolbar::SetComponentFontSize(uint8 FontSize)
+	{
+		UDMXPixelMapping* PixelMapping = WeakToolkit.IsValid() ? WeakToolkit.Pin()->GetDMXPixelMapping() : nullptr;
+		if (PixelMapping)
+		{
+			constexpr uint8 MinFontSize = 1;
+			PixelMapping->ComponentLabelFontSize = FMath::Max(FontSize, MinFontSize);
+		}
+	}
 }
 
 #undef LOCTEXT_NAMESPACE
diff --git a/Engine/Plugins/VirtualProduction/DMX/DMXPixelMapping/Source/DMXPixelMappingEditor/Private/Widgets/SDMXPixelMappingDesignerToolbar.h b/Engine/Plugins/VirtualProduction/DMX/DMXPixelMapping/Source/DMXPixelMappingEditor/Private/Widgets/SDMXPixelMappingDesignerToolbar.h
--- a/Engine/Plugins/VirtualProduction/DMX/DMXPixelMapping/Source/DMXPixelMappingEditor/Private/Widgets/SDMXPixelMappingDesignerToolbar.h
+++ b/Engine/Plugins/VirtualProduction/DMX/DMXPixelMapping/Source/DMXPixelMappingEditor/Private/Widgets/SDMXPixelMappingDesignerToolbar.h
@@ -8,6 +8,7 @@
 
 enum class ECheckBoxState : uint8;
 class FDMXPixelMappingToolkit;
+class FReply;
 namespace UE::DMX { enum class EDMXPixelMappingTransformHandleMode : uint8; }
 
 
@@ -25,14 +26,14 @@ namespace UE::DMX
 
 		SLATE_END_ARGS()
 
-		~SDMXPixelMappingDesignerToolbar();
-
 		/** Constructs this widget */
 		void Construct(const FArguments& InArgs, const TSharedRef<FDMXPixelMappingToolkit>& InToolkit);
 
 	private:
-		/** Builds the toolbar for this widget */
-		void RegisterToolbarMenu(const FArguments& InArgs);
+		TSharedRef<SWidget> GenerateTransformHandleModeSection();
+		TSharedRef<SWidget> GenerateGridSnappingSection();
+		TSharedRef<SWidget> GenerateZoomToFitSection(const FArguments& InArgs);
+		TSharedRef<SWidget> GenerateSettingsSection();
 
 		/** Generates the grid snapping menu */
 		TSharedRef<SWidget> GenerateSnapGridMenu();
@@ -52,6 +53,18 @@ namespace UE::DMX
 		/** Returns the current checkbox state of the transform mode */
 		ECheckBoxState GetCheckboxStateForTransormHandleMode(UE::DMX::EDMXPixelMappingTransformHandleMode TransformHandleMode) const;
 
+		/** Generates the contents of the settings menu */
+		TSharedRef<SWidget> GenerateSettingsMenuContent();
+
+		/** Generates a widget to edit the font size */
+		TSharedRef<SWidget> GenerateComponentFontSizeEditWidget();
+
+		/** Gets the font size for component labels in the designer view */
+		TOptional<uint8> GetComponentFontSize() const;
+		
+		/** Sets the font size for component labels in the designer view */
+		void SetComponentFontSize(uint8 FontSize);
+
 		/** The toolkit that owns this widget */
 		TWeakPtr<FDMXPixelMappingToolkit> WeakToolkit;
 	};
diff --git a/Engine/Plugins/VirtualProduction/DMX/DMXPixelMapping/Source/DMXPixelMappingEditor/Private/Widgets/SDMXPixelMappingOutputComponent.cpp b/Engine/Plugins/VirtualProduction/DMX/DMXPixelMapping/Source/DMXPixelMappingEditor/Private/Widgets/SDMXPixelMappingOutputComponent.cpp
--- a/Engine/Plugins/VirtualProduction/DMX/DMXPixelMapping/Source/DMXPixelMappingEditor/Private/Widgets/SDMXPixelMappingOutputComponent.cpp
+++ b/Engine/Plugins/VirtualProduction/DMX/DMXPixelMapping/Source/DMXPixelMappingEditor/Private/Widgets/SDMXPixelMappingOutputComponent.cpp
@@ -118,7 +118,9 @@ namespace UE::DMX
 		const FLinearColor Color = Model->GetColor();
 
 		const FVector2D Size = Model->GetSize();
-		TArray<FVector2D> LinePoints
+		
+		// Draw the component box
+		TArray<FVector2D> BoxPoints
 		{
 			FVector2D::ZeroVector,
 			FVector2D(0.0, Size.Y),
@@ -127,17 +129,56 @@ namespace UE::DMX
 			FVector2D::ZeroVector
 		};
 
-		constexpr float LineThickness = 2.f;
 		constexpr bool bAntialias = true;
+		constexpr float BoxLineThickness = 2.f;
 		FSlateDrawElement::MakeLines(
 			OutDrawElements,
 			LayerId,
 			AllottedGeometry.ToPaintGeometry(),
-			LinePoints,
+			BoxPoints,
 			ESlateDrawEffect::None,
 			Color,
 			bAntialias,
-			LineThickness);
+			BoxLineThickness);
+
+		// Selectively draw the pivot
+		const UDMXPixelMappingEditorSettings* Settings = GetDefault<UDMXPixelMappingEditorSettings>();
+		if (Model->ShouldDrawPivot())
+		{
+			const float PivotLength = FMath::Min(Size.X / 16.0, Size.Y / 16.0); // Y is up
+			TArray<FVector2D> PivotPointsX
+			{
+				FVector2D(Size.X / 2.0, Size.Y / 2.0),
+				FVector2D(Size.X / 2.0 + PivotLength, Size.Y / 2.0)
+			};
+
+			TArray<FVector2D> PivotPointsY
+			{
+				FVector2D(Size.X / 2.0, Size.Y / 2.0),
+				FVector2D(Size.X / 2.0, Size.Y / 2.0 - PivotLength)
+			};
+
+			constexpr float PivotLineThickness = 1.f;
+			FSlateDrawElement::MakeLines(
+				OutDrawElements,
+				LayerId,
+				AllottedGeometry.ToPaintGeometry(),
+				PivotPointsX,
+				ESlateDrawEffect::None,
+				FLinearColor::Red,
+				bAntialias,
+				PivotLineThickness);
+
+			FSlateDrawElement::MakeLines(
+				OutDrawElements,
+				LayerId,
+				AllottedGeometry.ToPaintGeometry(),
+				PivotPointsY,
+				ESlateDrawEffect::None,
+				FLinearColor::Green,
+				bAntialias,
+				PivotLineThickness);
+		}
 
 		return LayerId + 1;
 	}
diff --git a/Engine/Plugins/VirtualProduction/DMX/DMXPixelMapping/Source/DMXPixelMappingEditor/Private/Widgets/SDMXPixelMappingOutputComponentText.cpp b/Engine/Plugins/VirtualProduction/DMX/DMXPixelMapping/Source/DMXPixelMappingEditor/Private/Widgets/SDMXPixelMappingOutputComponentText.cpp
--- a/Engine/Plugins/VirtualProduction/DMX/DMXPixelMapping/Source/DMXPixelMappingEditor/Private/Widgets/SDMXPixelMappingOutputComponentText.cpp
+++ b/Engine/Plugins/VirtualProduction/DMX/DMXPixelMapping/Source/DMXPixelMappingEditor/Private/Widgets/SDMXPixelMappingOutputComponentText.cpp
@@ -2,9 +2,11 @@
 
 #include "Widgets/SDMXPixelMappingOutputComponentText.h"
 
+#include "DMXPixelMapping.h"
 #include "Fonts/FontMeasure.h"
 #include "Framework/Application/SlateApplication.h"
 #include "Styling/CoreStyle.h"
+#include "Toolkits/DMXPixelMappingToolkit.h"
 #include "ViewModels/DMXPixelMappingOutputComponentModel.h"
 
 
@@ -17,11 +19,6 @@ void SDMXPixelMappingOutputComponentText::Construct(const FArguments& InArgs, co
 
 	WeakToolkit = InToolkit;
 	Model = MakeShared<UE::DMX::FDMXPixelMappingOutputComponentModel>(InToolkit, OutputComponent);
-	
-	FontSize = InArgs._FontSize;
-
-	SetRenderTransform(TAttribute<TOptional<FSlateRenderTransform>>::CreateSP(this, &SDMXPixelMappingOutputComponentText::GetRenderTransform));
-	SetRenderTransformPivot(FVector2D(0.5, 0.5));
 
 	// Don't clip if the name has to be painted above the component box, otherwise clip. 
 	const EWidgetClipping WidgetClipping = Model->ShouldDrawNameAbove() ? EWidgetClipping::Inherit : EWidgetClipping::ClipToBoundsAlways;
@@ -63,35 +60,21 @@ int32 SDMXPixelMappingOutputComponentText::OnPaint(
 	return LayerId + 1;
 }
 
-TOptional<FSlateRenderTransform> SDMXPixelMappingOutputComponentText::GetRenderTransform() const
-{
-	const FQuat2D Quaternion = Model->GetQuaternion();
-
-	// The text always should face upwards to be readable. Hence,
-	// if the rotation goes > 90° or < -90°, draw the widget inverted.
-	if (Quaternion.GetVector().X < 0.0)
-	{
-		return FSlateRenderTransform(FQuat2D(UE_PI), FVector2D::ZeroVector);
-	}
-
-	return TOptional<FSlateRenderTransform>();
-}
-
 void SDMXPixelMappingOutputComponentText::OnPaintComponentName(const FPaintArgs& Args, const FGeometry& AllottedGeometry, FSlateWindowElementList& OutDrawElements, int32 LayerId, const FVector2f& Scale) const
 {
 	const float Margin = 8.f / Scale.Y;
 	const FVector2f LocalSize = AllottedGeometry.GetLocalSize();
 
 	// Upscale the font so it stays the same size, regardless of scaling (zoom)
-	const float ScaledFontSize = FontSize / Scale.Y;
+	const float ScaledFontSize = GetFontSize() / Scale.Y;
 
-	const float MaxFontSize = LocalSize.Y - ScaledFontSize - Margin * 2.f;
-	const float DesiredFontSize = ScaledFontSize < MaxFontSize ? ScaledFontSize : MaxFontSize;
-	const float ClampedFontSize = FMath::Max(UE_DELTA, DesiredFontSize);
+	// Always take the upper third section
+	const float MaxFontSize = LocalSize.Y - ScaledFontSize;
+	const float ClampedFontSize = FMath::Clamp(ScaledFontSize, UE_DELTA, MaxFontSize);
 
 	const float OffsetX = Margin;
-	const float OffsetY = Model->ShouldDrawNameAbove() ? -ScaledFontSize - Margin : Margin;
-
+	const float OffsetY = Model->ShouldDrawNameAbove() ? -ScaledFontSize * 1.5f : Margin;
+	
 	const FSlateLayoutTransform LayoutTransform(FVector2f(OffsetX, OffsetY));
 
 	FSlateDrawElement::MakeText(
@@ -99,7 +82,9 @@ void SDMXPixelMappingOutputComponentText::OnPaintComponentName(const FPaintArgs&
 		LayerId,
 		AllottedGeometry.ToPaintGeometry(LayoutTransform),
 		Model->GetName(),
-		FCoreStyle::GetDefaultFontStyle("Regular", ClampedFontSize));
+		FCoreStyle::GetDefaultFontStyle("Regular", ClampedFontSize),
+		ESlateDrawEffect::None,
+		FLinearColor::White.CopyWithNewOpacity(.8f));
 }
 
 void SDMXPixelMappingOutputComponentText::OnPaintCellID(const FPaintArgs& Args, const FGeometry& AllottedGeometry, FSlateWindowElementList& OutDrawElements, int32 LayerId, const FVector2f& Scale) const
@@ -108,11 +93,11 @@ void SDMXPixelMappingOutputComponentText::OnPaintCellID(const FPaintArgs& Args,
 	const FVector2f LocalSize = AllottedGeometry.GetLocalSize();
 
 	// Upscale the font so it stays the same size, regardless of scaling (zoom)
-	const float ScaledFontSize = FontSize / Scale.Y;
+	const float ScaledFontSize = GetFontSize() / Scale.Y;
 
+	// Draw in the center
 	const float MaxFontSize = LocalSize.Y - ScaledFontSize - Margin * 2.f;
-	const float DesiredFontSize = ScaledFontSize < MaxFontSize ? ScaledFontSize : MaxFontSize;
-	const float ClampedFontSize = FMath::Max(UE_DELTA, DesiredFontSize);
+	const float ClampedFontSize = FMath::Clamp(ScaledFontSize, UE_DELTA, MaxFontSize);
 
 	const FText CellIDText = Model->GetCellIDText();
 	const FSlateFontInfo Font = FCoreStyle::GetDefaultFontStyle("Regular", ScaledFontSize);
@@ -138,13 +123,12 @@ void SDMXPixelMappingOutputComponentText::OnPaintPatchInfo(const FPaintArgs& Arg
 	const float LocalSizeY = AllottedGeometry.GetLocalSize().Y;
 
 	// Upscale the font so it stays the same size, regardless of scaling (zoom)
-	const float ScaledFontSize = FontSize / Scale.Y;
+	const float ScaledFontSize = GetFontSize() / Scale.Y;
 	const float LineSpacing = 4.f / Scale.Y;
 
-	// Find a font size that fits in the lower third of this section. Consider two lines of text are painted
+	// Find a font size that fits below the component name. Consider two lines of text are painted
 	const float MaxFontSize = LocalSizeY - ScaledFontSize * 2.f - LineSpacing - Margin * 2.f;
-	const float DesiredFontSize = ScaledFontSize < MaxFontSize ? ScaledFontSize : MaxFontSize;
-	const float ClampedFontSize = FMath::Max(UE_DELTA, DesiredFontSize);
+	const float ClampedFontSize = FMath::Clamp(ScaledFontSize, UE_DELTA, MaxFontSize);
 
 	// Arrange both text elements in the lower section, align bottom
 	const float FixtureIDTextOffsetY = LocalSizeY - ClampedFontSize * 2.f - Margin - LineSpacing;
@@ -157,7 +141,7 @@ void SDMXPixelMappingOutputComponentText::OnPaintPatchInfo(const FPaintArgs& Arg
 		Model->GetFixtureIDText(),
 		FCoreStyle::GetDefaultFontStyle("Regular", ClampedFontSize),
 		ESlateDrawEffect::None,
-		FLinearColor::Green);
+		FLinearColor::Green.CopyWithNewOpacity(.8f));
 
 	const float AddressesTextOffsetY = LocalSizeY - ClampedFontSize - Margin;
 	const FSlateLayoutTransform AddressesLayoutTransform(FVector2f(Margin, AddressesTextOffsetY));
@@ -169,5 +153,16 @@ void SDMXPixelMappingOutputComponentText::OnPaintPatchInfo(const FPaintArgs& Arg
 		Model->GetAddressesText(),
 		FCoreStyle::GetDefaultFontStyle("Regular", ClampedFontSize),
 		ESlateDrawEffect::None,
-		FLinearColor::Green);
+		FLinearColor::Green.CopyWithNewOpacity(.8f));
+}
+
+uint8 SDMXPixelMappingOutputComponentText::GetFontSize() const
+{
+	UDMXPixelMapping* PixelMapping = WeakToolkit.IsValid() ? WeakToolkit.Pin()->GetDMXPixelMapping() : nullptr;
+	if (PixelMapping)
+	{
+		return PixelMapping->ComponentLabelFontSize;
+	}
+
+	return 8;
 }
diff --git a/Engine/Plugins/VirtualProduction/DMX/DMXPixelMapping/Source/DMXPixelMappingEditor/Private/Widgets/SDMXPixelMappingOutputComponentText.h b/Engine/Plugins/VirtualProduction/DMX/DMXPixelMapping/Source/DMXPixelMappingEditor/Private/Widgets/SDMXPixelMappingOutputComponentText.h
--- a/Engine/Plugins/VirtualProduction/DMX/DMXPixelMapping/Source/DMXPixelMappingEditor/Private/Widgets/SDMXPixelMappingOutputComponentText.h
+++ b/Engine/Plugins/VirtualProduction/DMX/DMXPixelMapping/Source/DMXPixelMappingEditor/Private/Widgets/SDMXPixelMappingOutputComponentText.h
@@ -2,6 +2,7 @@
 
 #pragma once
 
+#include "Math/TransformCalculus2D.h"
 #include "Widgets/SCompoundWidget.h"
 
 class FDMXPixelMappingToolkit;
@@ -16,10 +17,7 @@ class SDMXPixelMappingOutputComponentText
 
 public:
 	SLATE_BEGIN_ARGS(SDMXPixelMappingOutputComponentText)
-		: _FontSize(8)
 		{}
-		
-		SLATE_ARGUMENT(int32, FontSize)
 
 	SLATE_END_ARGS()
 
@@ -38,9 +36,6 @@ protected:
 	//~ End SWidget Interface
 
 private:
-	/** Returns the render transform for this widget */
-	TOptional<FSlateRenderTransform> GetRenderTransform() const;
-
 	/** Helper to paint the component name */
 	void OnPaintComponentName(const FPaintArgs& Args, const FGeometry& AllottedGeometry, FSlateWindowElementList& OutDrawElements, int32 LayerId, const FVector2f& Scale) const;
 
@@ -50,6 +45,6 @@ private:
 	/** Helper to paint the patch info */
 	void OnPaintPatchInfo(const FPaintArgs& Args, const FGeometry& AllottedGeometry, FSlateWindowElementList& OutDrawElements, int32 LayerId, const FVector2f& Scale) const;
 
-	// Slate args
-	int32 FontSize;
+	/** Returns the current font size */
+	uint8 GetFontSize() const;
 };
diff --git a/Engine/Plugins/VirtualProduction/DMX/DMXPixelMapping/Source/DMXPixelMappingEditor/Private/Widgets/SDMXPixelMappingSurface.h b/Engine/Plugins/VirtualProduction/DMX/DMXPixelMapping/Source/DMXPixelMappingEditor/Private/Widgets/SDMXPixelMappingSurface.h
--- a/Engine/Plugins/VirtualProduction/DMX/DMXPixelMapping/Source/DMXPixelMappingEditor/Private/Widgets/SDMXPixelMappingSurface.h
+++ b/Engine/Plugins/VirtualProduction/DMX/DMXPixelMapping/Source/DMXPixelMappingEditor/Private/Widgets/SDMXPixelMappingSurface.h
@@ -101,10 +101,10 @@ protected:
 	int32 ZoomLevel;
 
 	/** Are we panning the view at the moment? */
-	bool bIsPanning;
+	bool bIsPanning = false;
 
 	/** Are we zooming the view with trackpad at the moment? */
-	bool bIsZooming;
+	bool bIsZooming = false;
 
 	/** Allow continuous zoom interpolation? */
 	TAttribute<bool> AllowContinousZoomInterpolation;
diff --git a/Engine/Plugins/VirtualProduction/DMX/DMXPixelMapping/Source/DMXPixelMappingEditor/Private/Widgets/SDMXPixelMappingTransformHandle.cpp b/Engine/Plugins/VirtualProduction/DMX/DMXPixelMapping/Source/DMXPixelMappingEditor/Private/Widgets/SDMXPixelMappingTransformHandle.cpp
--- a/Engine/Plugins/VirtualProduction/DMX/DMXPixelMapping/Source/DMXPixelMappingEditor/Private/Widgets/SDMXPixelMappingTransformHandle.cpp
+++ b/Engine/Plugins/VirtualProduction/DMX/DMXPixelMapping/Source/DMXPixelMappingEditor/Private/Widgets/SDMXPixelMappingTransformHandle.cpp
@@ -250,21 +250,15 @@ FVector2D SDMXPixelMappingTransformHandle::GetSnapSize(UDMXPixelMappingOutputCom
 	UDMXPixelMappingRendererComponent* RendererComponent = OutputComponent->GetRendererComponent();
 
 	const UDMXPixelMapping* PixelMapping = Toolkit->GetDMXPixelMapping();
-	if (!PixelMapping)
+	if (!PixelMapping || !PixelMapping->bGridSnappingEnabled)
 	{
 		return RequestedSize;
 	}
 
 	const FVector2D CellSize = [PixelMapping, RendererComponent]()
 		{
-			if (PixelMapping->bGridSnappingEnabled)
-			{
-				const FVector2D TextureSize = RendererComponent->GetSize();
-				return TextureSize / FVector2D(PixelMapping->SnapGridColumns, PixelMapping->SnapGridRows);
-			}
-
-			// Grid snap to pixels if grid snapping is disabled
-			return FVector2D(1.f, 1.f);
+			const FVector2D TextureSize = RendererComponent->GetSize();
+			return TextureSize / FVector2D(PixelMapping->SnapGridColumns, PixelMapping->SnapGridRows);
 		}();
 
 	// Grid snap bottom right
diff --git a/Engine/Plugins/VirtualProduction/DMX/DMXPixelMapping/Source/DMXPixelMappingRuntime/Public/DMXPixelMapping.h b/Engine/Plugins/VirtualProduction/DMX/DMXPixelMapping/Source/DMXPixelMappingRuntime/Public/DMXPixelMapping.h
--- a/Engine/Plugins/VirtualProduction/DMX/DMXPixelMapping/Source/DMXPixelMappingRuntime/Public/DMXPixelMapping.h
+++ b/Engine/Plugins/VirtualProduction/DMX/DMXPixelMapping/Source/DMXPixelMappingRuntime/Public/DMXPixelMapping.h
@@ -145,21 +145,25 @@ public:
 
 #if WITH_EDITORONLY_DATA
 	/** If true, grid snapping is enabled */
-	UPROPERTY()
+	UPROPERTY(NonTransactional)
 	bool bGridSnappingEnabled = false;
 
 	/** The number of columns in the grid */
-	UPROPERTY()
+	UPROPERTY(NonTransactional)
 	int32 SnapGridColumns = 10;
 
 	/** The number of rows in the grid */
-	UPROPERTY()
+	UPROPERTY(NonTransactional)
 	int32 SnapGridRows = 10;
 
 	/** The color of the grid snapping grid */
-	UPROPERTY()
+	UPROPERTY(NonTransactional)
 	FLinearColor SnapGridColor;
 
+	/** Font size for the component labels in the designer view */
+	UPROPERTY(NonTransactional)
+	float ComponentLabelFontSize = 8.f;
+
 	/** If true, editor is set to scale children with parent. This is forwarded from the editor module (DMXPixelMappingEditorSettings) to be accessible in the runtime module. */
 	UPROPERTY(Transient, NonTransactional)
 	bool bEditorScaleChildrenWithParent = false;
