DMX - PixelMapping hide transform handles on component sides when in rotate mode

diff --git a/Engine/Plugins/VirtualProduction/DMX/DMXPixelMapping/Source/DMXPixelMappingEditor/Private/Widgets/SDMXPixelMappingTransformHandle.cpp b/Engine/Plugins/VirtualProduction/DMX/DMXPixelMapping/Source/DMXPixelMappingEditor/Private/Widgets/SDMXPixelMappingTransformHandle.cpp
--- a/Engine/Plugins/VirtualProduction/DMX/DMXPixelMapping/Source/DMXPixelMappingEditor/Private/Widgets/SDMXPixelMappingTransformHandle.cpp
+++ b/Engine/Plugins/VirtualProduction/DMX/DMXPixelMapping/Source/DMXPixelMappingEditor/Private/Widgets/SDMXPixelMappingTransformHandle.cpp
@@ -26,17 +26,29 @@ void SDMXPixelMappingTransformHandle::Construct(const FArguments& InArgs, TShare
 	DragDirection = ComputeDragDirection(InTransformDirection);
 	DragOrigin = ComputeOrigin(InTransformDirection);
 
+	SetVisibility(TAttribute<EVisibility>::CreateSP(this, &SDMXPixelMappingTransformHandle::GetHandleVisibility));
+
 	ChildSlot
 	[
 		SNew(SImage)
-		.Visibility(this, &SDMXPixelMappingTransformHandle::GetHandleVisibility)
 		.Image(FAppStyle::Get().GetBrush("UMGEditor.TransformHandle"))
 	];
 }
 
 EVisibility SDMXPixelMappingTransformHandle::GetHandleVisibility() const
 {
-	return EVisibility::Visible;
+	if (TransformDirection == EDMXPixelMappingTransformDirection::BottomRight)
+	{
+		return EVisibility::Visible;
+	}
+
+	const TSharedPtr<FDMXPixelMappingToolkit> Toolkit = DesignerViewWeakPtr.IsValid() ? DesignerViewWeakPtr.Pin()->GetToolkit() : nullptr;
+	if (Toolkit.IsValid() && Toolkit->GetTransformHandleMode() == UE::DMX::EDMXPixelMappingTransformHandleMode::Resize)
+	{
+		return EVisibility::Visible;
+	}
+
+	return EVisibility::Hidden;
 }
 
 FReply SDMXPixelMappingTransformHandle::OnMouseButtonDown(const FGeometry& MyGeometry, const FPointerEvent& MouseEvent)
