DMX - Fix a PixelMapping bug when multi drag-drop components with grid snapping enabled, fix a bug where matrix components do not properly transact

diff --git a/Engine/Plugins/VirtualProduction/DMX/DMXPixelMapping/Source/DMXPixelMappingEditor/Private/DragDrop/DMXPixelMappingDragDropOp.cpp b/Engine/Plugins/VirtualProduction/DMX/DMXPixelMapping/Source/DMXPixelMappingEditor/Private/DragDrop/DMXPixelMappingDragDropOp.cpp
--- a/Engine/Plugins/VirtualProduction/DMX/DMXPixelMapping/Source/DMXPixelMappingEditor/Private/DragDrop/DMXPixelMappingDragDropOp.cpp
+++ b/Engine/Plugins/VirtualProduction/DMX/DMXPixelMapping/Source/DMXPixelMappingEditor/Private/DragDrop/DMXPixelMappingDragDropOp.cpp
@@ -84,12 +84,16 @@ void FDMXPixelMappingDragDropOp::LayoutOutputComponents(const FVector2D& GraphSp
 	{
 		if (UDMXPixelMappingOutputComponent* OutputComponent = Cast<UDMXPixelMappingOutputComponent>(Component.Get()))
 		{
+			OutputComponent->PreEditChange(nullptr);
+
 			if (ensureMsgf(OutputComponent->GetClass() != UDMXPixelMappingMatrixComponent::StaticClass(),
 				TEXT("Matrix components cannot be laid out with FDMXPixelMappingDragDropOp::LayoutOutputComponents. Please use FGroupChildDragDropHelper instead (see FDMXPixelMappingDragDropOp::GetGroupChildDragDropHelper().")))
 			{
 				OutputComponent->Modify();
 				OutputComponent->SetPosition(OutputComponent->GetPosition() + Translation);
 			}
+
+			OutputComponent->PostEditChange();
 		}
 	}
 }
diff --git a/Engine/Plugins/VirtualProduction/DMX/DMXPixelMapping/Source/DMXPixelMappingEditor/Private/DragDrop/DMXPixelMappingGroupChildDragDropHelper.cpp b/Engine/Plugins/VirtualProduction/DMX/DMXPixelMapping/Source/DMXPixelMappingEditor/Private/DragDrop/DMXPixelMappingGroupChildDragDropHelper.cpp
--- a/Engine/Plugins/VirtualProduction/DMX/DMXPixelMapping/Source/DMXPixelMappingEditor/Private/DragDrop/DMXPixelMappingGroupChildDragDropHelper.cpp
+++ b/Engine/Plugins/VirtualProduction/DMX/DMXPixelMapping/Source/DMXPixelMappingEditor/Private/DragDrop/DMXPixelMappingGroupChildDragDropHelper.cpp
@@ -159,33 +159,22 @@ void FDMXPixelMappingGroupChildDragDropHelper::LayoutUnaligned(const FVector2D&
 		return;
 	}
 
+	// Find where the first component grid snaps
 	const FVector2D Anchor = FirstComponent->GetPositionRotated();
-	for (const TWeakObjectPtr<UDMXPixelMappingOutputComponent>& WeakOutputComponent : WeakChildComponents)
-	{
-		if (UDMXPixelMappingOutputComponent* ChildComponent = WeakOutputComponent.Get())
-		{
-			ChildComponent->PreEditChange(nullptr);
-
-			constexpr bool bModifyChildrenRecursive = true;
-			ChildComponent->ForEachChild([](UDMXPixelMappingBaseComponent* Component)
-				{
-					Component->Modify();
-				}, bModifyChildrenRecursive);
-
-			const FVector2D AnchorOffset = Anchor - ChildComponent->GetPositionRotated();
+	const FVector2D AnchorOffset = Anchor - FirstComponent->GetPositionRotated();
 
-			const FVector2D NewPosition = GraphSpacePosition - AnchorOffset - DragDropOp->GraphSpaceDragOffset;
-			ChildComponent->SetPositionRotated(NewPosition);
-		}
-	}
-	const FVector2D GridSnapPosition = DragDropOp->ComputeGridSnapPosition(FirstComponent->GetPositionRotated());
-	FirstComponent->SetPositionRotated(GridSnapPosition);
+	const FVector2D DesiredPosition = GraphSpacePosition - AnchorOffset - DragDropOp->GraphSpaceDragOffset;
+	const FVector2D GridSnapPosition = DragDropOp->ComputeGridSnapPosition(DesiredPosition);
 
+	// Translate all
+	const FVector2D Translation = GridSnapPosition - Anchor;
 	for (const TWeakObjectPtr<UDMXPixelMappingOutputComponent>& WeakOutputComponent : WeakChildComponents)
 	{
 		if (UDMXPixelMappingOutputComponent* ChildComponent = WeakOutputComponent.Get())
 		{
-			ChildComponent->PostEditChange();
+			ChildComponent->Modify();
+
+			ChildComponent->SetPosition(ChildComponent->GetPosition() + Translation);
 		}
 	}
 }
diff --git a/Engine/Plugins/VirtualProduction/DMX/DMXPixelMapping/Source/DMXPixelMappingRuntime/Private/Components/DMXPixelMappingComponentGeometryCache.cpp b/Engine/Plugins/VirtualProduction/DMX/DMXPixelMapping/Source/DMXPixelMappingRuntime/Private/Components/DMXPixelMappingComponentGeometryCache.cpp
--- a/Engine/Plugins/VirtualProduction/DMX/DMXPixelMapping/Source/DMXPixelMappingRuntime/Private/Components/DMXPixelMappingComponentGeometryCache.cpp
+++ b/Engine/Plugins/VirtualProduction/DMX/DMXPixelMapping/Source/DMXPixelMappingRuntime/Private/Components/DMXPixelMappingComponentGeometryCache.cpp
@@ -130,7 +130,7 @@ void FDMXPixelMappingComponentGeometryCache::PropagonateSizeChangesToChildren(co
 
 		constexpr bool bRecursive = false;
 		OutputComponent->ForEachChildOfClass<UDMXPixelMappingOutputComponent>(
-			[Scalar, DeltaSize, this](UDMXPixelMappingOutputComponent* Child)
+			[Scalar, this](UDMXPixelMappingOutputComponent* Child)
 			{
 				const FVector2D OldChildPosition = Child->GetPosition();
 				const FVector2D OldChildSize = Child->GetSize();
diff --git a/Engine/Plugins/VirtualProduction/DMX/DMXPixelMapping/Source/DMXPixelMappingRuntime/Private/Components/DMXPixelMappingMatrixComponent.cpp b/Engine/Plugins/VirtualProduction/DMX/DMXPixelMapping/Source/DMXPixelMappingRuntime/Private/Components/DMXPixelMappingMatrixComponent.cpp
--- a/Engine/Plugins/VirtualProduction/DMX/DMXPixelMapping/Source/DMXPixelMappingRuntime/Private/Components/DMXPixelMappingMatrixComponent.cpp
+++ b/Engine/Plugins/VirtualProduction/DMX/DMXPixelMapping/Source/DMXPixelMappingRuntime/Private/Components/DMXPixelMappingMatrixComponent.cpp
@@ -444,6 +444,8 @@ void UDMXPixelMappingMatrixComponent::HandleMatrixChanged()
 		return;
 	}
 
+	const double RestoreRotation = GetRotation();
+	SetRotation(0.0);
 	CoordinateGrid = 0;
 
 	// Remove all existing children and rebuild them anew
@@ -492,7 +494,11 @@ void UDMXPixelMappingMatrixComponent::HandleMatrixChanged()
 		},
 		bUpdateSizeRecursive);
 
+	SetRotation(RestoreRotation);
+
+	PRAGMA_DISABLE_DEPRECATION_WARNINGS
 	GetOnMatrixChanged().Broadcast(PixelMapping, this);
+	PRAGMA_ENABLE_DEPRECATION_WARNINGS
 }
 
 void UDMXPixelMappingMatrixComponent::QueueDownsample()
diff --git a/Engine/Plugins/VirtualProduction/DMX/DMXPixelMapping/Source/DMXPixelMappingRuntime/Public/Components/DMXPixelMappingMatrixComponent.h b/Engine/Plugins/VirtualProduction/DMX/DMXPixelMapping/Source/DMXPixelMappingRuntime/Public/Components/DMXPixelMappingMatrixComponent.h
--- a/Engine/Plugins/VirtualProduction/DMX/DMXPixelMapping/Source/DMXPixelMappingRuntime/Public/Components/DMXPixelMappingMatrixComponent.h
+++ b/Engine/Plugins/VirtualProduction/DMX/DMXPixelMapping/Source/DMXPixelMappingRuntime/Public/Components/DMXPixelMappingMatrixComponent.h
@@ -35,10 +35,11 @@ public:
 	UDMXPixelMappingMatrixComponent();
 
 	/** Gets an Event broadcast when a the matrix (and by that its num cells) changed */
+	UE_DEPRECATED(5.4, "OnMatrixChanged is removed without replacement. This is handled internally and should not require external handling.")
 	static FDMXPixelMappingOnMatrixChanged& GetOnMatrixChanged()
 	{
-		static FDMXPixelMappingOnMatrixChanged OnMatrixChanged;
-		return OnMatrixChanged;
+		static FDMXPixelMappingOnMatrixChanged OnMatrixChanged_DEPRECATED;
+		return OnMatrixChanged_DEPRECATED;
 	}
 
 	// ~Begin UObject interface
