DMX - Add an option to flip matrix cells, useful when a matrix is facing away from the screen

diff --git a/Engine/Plugins/VirtualProduction/DMX/DMXPixelMapping/Source/DMXPixelMappingEditor/Private/Customizations/DMXPixelMappingDetailCustomization_Matrix.cpp b/Engine/Plugins/VirtualProduction/DMX/DMXPixelMapping/Source/DMXPixelMappingEditor/Private/Customizations/DMXPixelMappingDetailCustomization_Matrix.cpp
--- a/Engine/Plugins/VirtualProduction/DMX/DMXPixelMapping/Source/DMXPixelMappingEditor/Private/Customizations/DMXPixelMappingDetailCustomization_Matrix.cpp
+++ b/Engine/Plugins/VirtualProduction/DMX/DMXPixelMapping/Source/DMXPixelMappingEditor/Private/Customizations/DMXPixelMappingDetailCustomization_Matrix.cpp
@@ -81,6 +81,12 @@ void FDMXPixelMappingDetailCustomization_Matrix::CustomizeDetails(IDetailLayoutB
 				(*TransformCategoryPtr)->SetSortOrder(SortOrder++);
 			}
 
+			IDetailCategoryBuilder* const* MatrixCategoryPtr = CategoryMap.Find("Matrix");
+			if (MatrixCategoryPtr)
+			{
+				(*MatrixCategoryPtr)->SetSortOrder(SortOrder++);
+			}
+
 			IDetailCategoryBuilder* const* EditorSettingsCategoryPtr = CategoryMap.Find("Editor Settings");
 			if (EditorSettingsCategoryPtr)
 			{
diff --git a/Engine/Plugins/VirtualProduction/DMX/DMXPixelMapping/Source/DMXPixelMappingRuntime/Private/Components/DMXPixelMappingMatrixComponent.cpp b/Engine/Plugins/VirtualProduction/DMX/DMXPixelMapping/Source/DMXPixelMappingRuntime/Private/Components/DMXPixelMappingMatrixComponent.cpp
--- a/Engine/Plugins/VirtualProduction/DMX/DMXPixelMapping/Source/DMXPixelMappingRuntime/Private/Components/DMXPixelMappingMatrixComponent.cpp
+++ b/Engine/Plugins/VirtualProduction/DMX/DMXPixelMapping/Source/DMXPixelMappingRuntime/Private/Components/DMXPixelMappingMatrixComponent.cpp
@@ -183,6 +183,8 @@ void UDMXPixelMappingMatrixComponent::PostEditChangeProperty(FPropertyChangedEve
 	}
 	else if (PropertyName == GET_MEMBER_NAME_CHECKED(UDMXPixelMappingMatrixComponent, FixturePatchRef) ||
 		PropertyName == GET_MEMBER_NAME_CHECKED(UDMXPixelMappingMatrixComponent, CoordinateGrid) ||
+		PropertyName == GET_MEMBER_NAME_CHECKED(UDMXPixelMappingMatrixComponent, bInvertCellsX) ||
+		PropertyName == GET_MEMBER_NAME_CHECKED(UDMXPixelMappingMatrixComponent, bInvertCellsY) ||
 		PropertyName == GET_MEMBER_NAME_CHECKED(FDMXEntityReference, DMXLibrary))
 	{
 		HandleMatrixChanged();
@@ -485,12 +487,21 @@ void UDMXPixelMappingMatrixComponent::HandleMatrixChanged()
 		CellSize = FVector2D(GetSize().X / CoordinateGrid.X, GetSize().Y / CoordinateGrid.Y);
 	}
 
+	auto InvertCellsIfRequiredLambda = [this](FIntPoint CellCoordinate)
+		{
+			CellCoordinate.X = bInvertCellsX ? FMath::Abs(CellCoordinate.X - CoordinateGrid.X) - 1 : CellCoordinate.X;
+			CellCoordinate.Y = bInvertCellsY ? FMath::Abs(CellCoordinate.Y - CoordinateGrid.Y) - 1 : CellCoordinate.Y;
+			return CellCoordinate;
+		};
+
 	constexpr bool bUpdateSizeRecursive = false;
-	ForEachChildOfClass<UDMXPixelMappingMatrixCellComponent>([this](UDMXPixelMappingMatrixCellComponent* ChildComponent)
+	ForEachChildOfClass<UDMXPixelMappingMatrixCellComponent>([this, &InvertCellsIfRequiredLambda](UDMXPixelMappingMatrixCellComponent* ChildComponent)
 		{
 			ChildComponent->SetRotation(GetRotation());
 			ChildComponent->SetSize(CellSize);
-			ChildComponent->SetPosition(GetPosition() + FVector2D(CellSize * ChildComponent->GetCellCoordinate()));
+
+			const FIntPoint CellCoordinate = InvertCellsIfRequiredLambda(ChildComponent->GetCellCoordinate());
+			ChildComponent->SetPosition(GetPosition() + FVector2D(CellSize * CellCoordinate));
 		},
 		bUpdateSizeRecursive);
 
diff --git a/Engine/Plugins/VirtualProduction/DMX/DMXPixelMapping/Source/DMXPixelMappingRuntime/Public/Components/DMXPixelMappingMatrixComponent.h b/Engine/Plugins/VirtualProduction/DMX/DMXPixelMapping/Source/DMXPixelMappingRuntime/Public/Components/DMXPixelMappingMatrixComponent.h
--- a/Engine/Plugins/VirtualProduction/DMX/DMXPixelMapping/Source/DMXPixelMappingRuntime/Public/Components/DMXPixelMappingMatrixComponent.h
+++ b/Engine/Plugins/VirtualProduction/DMX/DMXPixelMapping/Source/DMXPixelMappingRuntime/Public/Components/DMXPixelMappingMatrixComponent.h
@@ -115,6 +115,14 @@ public:
 	UPROPERTY()
 	EDMXPixelMappingDistribution Distribution;
 
+	/** Inverts the cell ordering along the X-axis. Useful when the matrix is facing away from the screen, rotated horizontally. */
+	UPROPERTY(EditAnywhere, Category = "Matrix")
+	bool bInvertCellsX = false;
+
+	/** Inverts the cell ordering along the Y-axis. Useful when the matrix is facing away from the screen, rotated vertically. */
+	UPROPERTY(EditAnywhere, Category = "Matrix")
+	bool bInvertCellsY = false;
+
 	/** Layout script for the children of this component (hidden in customizations and displayed in its own panel). */
 	UPROPERTY(EditAnywhere, Instanced, Category = "Layout")
 	TObjectPtr<UDMXPixelMappingLayoutScript> LayoutScript;
