DMX - PixelMapping now has an option to use the patch color as the component color

diff --git a/Engine/Plugins/VirtualProduction/DMX/DMXPixelMapping/Source/DMXPixelMappingEditor/Private/Customizations/DMXPixelMappingDetailCustomization_FixtureGroupItem.cpp b/Engine/Plugins/VirtualProduction/DMX/DMXPixelMapping/Source/DMXPixelMappingEditor/Private/Customizations/DMXPixelMappingDetailCustomization_FixtureGroupItem.cpp
--- a/Engine/Plugins/VirtualProduction/DMX/DMXPixelMapping/Source/DMXPixelMappingEditor/Private/Customizations/DMXPixelMappingDetailCustomization_FixtureGroupItem.cpp
+++ b/Engine/Plugins/VirtualProduction/DMX/DMXPixelMapping/Source/DMXPixelMappingEditor/Private/Customizations/DMXPixelMappingDetailCustomization_FixtureGroupItem.cpp
@@ -5,6 +5,7 @@
 #include "Components/DMXPixelMappingFixtureGroupItemComponent.h"
 #include "DetailCategoryBuilder.h"
 #include "DetailLayoutBuilder.h"
+#include "PropertyHandle.h"
 
 
 FDMXPixelMappingDetailCustomization_FixtureGroupItem::FDMXPixelMappingDetailCustomization_FixtureGroupItem(const TWeakPtr<FDMXPixelMappingToolkit> InToolkitWeakPtr)
@@ -18,6 +19,21 @@ TSharedRef<IDetailCustomization> FDMXPixelMappingDetailCustomization_FixtureGrou
 
 void FDMXPixelMappingDetailCustomization_FixtureGroupItem::CustomizeDetails(IDetailLayoutBuilder& InDetailLayout)
 {
+	// Disable the editor color property, if the component Uses the patch color
+	IDetailCategoryBuilder& CategoryBuilder = InDetailLayout.EditCategory("Editor Settings");
+
+	UsePatchColorHandle = InDetailLayout.GetProperty(GET_MEMBER_NAME_CHECKED(UDMXPixelMappingOutputDMXComponent, bUsePatchColor), UDMXPixelMappingOutputDMXComponent::StaticClass());
+	CategoryBuilder.AddProperty(UsePatchColorHandle);
+
+	PRAGMA_DISABLE_DEPRECATION_WARNINGS
+	// The EditorColor property should no longer be publicly accessed. However it's ok to access it here.
+	const FName EditorColorMemberName = GET_MEMBER_NAME_CHECKED(UDMXPixelMappingOutputComponent, EditorColor);
+	PRAGMA_ENABLE_DEPRECATION_WARNINGS
+
+	const TSharedRef<IPropertyHandle> EditorColorHandle = InDetailLayout.GetProperty(EditorColorMemberName, UDMXPixelMappingOutputComponent::StaticClass());
+	CategoryBuilder.AddProperty(EditorColorHandle)
+		.EditCondition(TAttribute<bool>::CreateSP(this, &FDMXPixelMappingDetailCustomization_FixtureGroupItem::CanEditEditorColor), FOnBooleanValueChanged());
+
 	// Sort categories
 	InDetailLayout.SortCategories([](const TMap<FName, IDetailCategoryBuilder*>& CategoryMap)
 		{
@@ -85,3 +101,14 @@ void FDMXPixelMappingDetailCustomization_FixtureGroupItem::CustomizeDetails(IDet
 			}
 		});
 }
+
+bool FDMXPixelMappingDetailCustomization_FixtureGroupItem::CanEditEditorColor() const
+{
+	bool bUsePatchColor;
+	if (UsePatchColorHandle->GetValue(bUsePatchColor) == FPropertyAccess::Success)
+	{
+		return !bUsePatchColor;
+	}
+
+	return true;
+}
diff --git a/Engine/Plugins/VirtualProduction/DMX/DMXPixelMapping/Source/DMXPixelMappingEditor/Private/Customizations/DMXPixelMappingDetailCustomization_FixtureGroupItem.h b/Engine/Plugins/VirtualProduction/DMX/DMXPixelMapping/Source/DMXPixelMappingEditor/Private/Customizations/DMXPixelMappingDetailCustomization_FixtureGroupItem.h
--- a/Engine/Plugins/VirtualProduction/DMX/DMXPixelMapping/Source/DMXPixelMappingEditor/Private/Customizations/DMXPixelMappingDetailCustomization_FixtureGroupItem.h
+++ b/Engine/Plugins/VirtualProduction/DMX/DMXPixelMapping/Source/DMXPixelMappingEditor/Private/Customizations/DMXPixelMappingDetailCustomization_FixtureGroupItem.h
@@ -6,6 +6,7 @@
 
 class FDMXPixelMappingToolkit;
 class IDetailLayoutBuilder;
+class IPropertyHandle;
 
 
 class FDMXPixelMappingDetailCustomization_FixtureGroupItem
@@ -21,6 +22,12 @@ public:
 	//~ IPropertyTypeCustomization interface end
 
 private:
+	/** Returns true if the EditorColor property is editable */
+	bool CanEditEditorColor() const;
+
+	/** Handle for the bUsePatchColor property */
+	TSharedPtr<IPropertyHandle> UsePatchColorHandle;
+
 	/** Weak reference to the DMX editor */
 	TWeakPtr<FDMXPixelMappingToolkit> ToolkitWeakPtr;
 };
diff --git a/Engine/Plugins/VirtualProduction/DMX/DMXPixelMapping/Source/DMXPixelMappingEditor/Private/Customizations/DMXPixelMappingDetailCustomization_Matrix.cpp b/Engine/Plugins/VirtualProduction/DMX/DMXPixelMapping/Source/DMXPixelMappingEditor/Private/Customizations/DMXPixelMappingDetailCustomization_Matrix.cpp
--- a/Engine/Plugins/VirtualProduction/DMX/DMXPixelMapping/Source/DMXPixelMappingEditor/Private/Customizations/DMXPixelMappingDetailCustomization_Matrix.cpp
+++ b/Engine/Plugins/VirtualProduction/DMX/DMXPixelMapping/Source/DMXPixelMappingEditor/Private/Customizations/DMXPixelMappingDetailCustomization_Matrix.cpp
@@ -5,6 +5,7 @@
 #include "Components/DMXPixelMappingMatrixComponent.h"
 #include "DetailCategoryBuilder.h"
 #include "DetailLayoutBuilder.h"
+#include "PropertyHandle.h"
 
 
 FDMXPixelMappingDetailCustomization_Matrix::FDMXPixelMappingDetailCustomization_Matrix(const TWeakPtr<FDMXPixelMappingToolkit> InToolkitWeakPtr)
@@ -18,9 +19,24 @@ TSharedRef<IDetailCustomization> FDMXPixelMappingDetailCustomization_Matrix::Mak
 
 void FDMXPixelMappingDetailCustomization_Matrix::CustomizeDetails(IDetailLayoutBuilder& InDetailLayout)
 {
-	// Hide the Layout Script property (shown in its own panel, see SDMXPixelMappingLayoutView)
+	// Hide the Layout Script property
 	InDetailLayout.HideProperty(GET_MEMBER_NAME_CHECKED(UDMXPixelMappingMatrixComponent, LayoutScript));
 
+	// Disable the editor color property, if the component Uses the patch color
+	IDetailCategoryBuilder& CategoryBuilder = InDetailLayout.EditCategory("Editor Settings");
+
+	UsePatchColorHandle = InDetailLayout.GetProperty(GET_MEMBER_NAME_CHECKED(UDMXPixelMappingOutputDMXComponent, bUsePatchColor), UDMXPixelMappingOutputDMXComponent::StaticClass());
+	CategoryBuilder.AddProperty(UsePatchColorHandle);
+
+	PRAGMA_DISABLE_DEPRECATION_WARNINGS
+	// The EditorColor property should no longer be publicly accessed. However it's ok to access it here.
+	const FName EditorColorMemberName = GET_MEMBER_NAME_CHECKED(UDMXPixelMappingOutputComponent, EditorColor);
+	PRAGMA_ENABLE_DEPRECATION_WARNINGS
+
+	const TSharedRef<IPropertyHandle> EditorColorHandle = InDetailLayout.GetProperty(EditorColorMemberName, UDMXPixelMappingOutputComponent::StaticClass());
+	CategoryBuilder.AddProperty(EditorColorHandle)
+		.EditCondition(TAttribute<bool>::CreateSP(this, &FDMXPixelMappingDetailCustomization_Matrix::CanEditEditorColor), FOnBooleanValueChanged());
+
 	// Sort categories
 	InDetailLayout.SortCategories([](const TMap<FName, IDetailCategoryBuilder*>& CategoryMap)
 		{
@@ -94,3 +110,14 @@ void FDMXPixelMappingDetailCustomization_Matrix::CustomizeDetails(IDetailLayoutB
 			}
 		});
 }
+
+bool FDMXPixelMappingDetailCustomization_Matrix::CanEditEditorColor() const
+{
+	bool bUsePatchColor;
+	if (UsePatchColorHandle->GetValue(bUsePatchColor) == FPropertyAccess::Success)
+	{
+		return !bUsePatchColor;
+	}
+
+	return true;
+}
diff --git a/Engine/Plugins/VirtualProduction/DMX/DMXPixelMapping/Source/DMXPixelMappingEditor/Private/Customizations/DMXPixelMappingDetailCustomization_Matrix.h b/Engine/Plugins/VirtualProduction/DMX/DMXPixelMapping/Source/DMXPixelMappingEditor/Private/Customizations/DMXPixelMappingDetailCustomization_Matrix.h
--- a/Engine/Plugins/VirtualProduction/DMX/DMXPixelMapping/Source/DMXPixelMappingEditor/Private/Customizations/DMXPixelMappingDetailCustomization_Matrix.h
+++ b/Engine/Plugins/VirtualProduction/DMX/DMXPixelMapping/Source/DMXPixelMappingEditor/Private/Customizations/DMXPixelMappingDetailCustomization_Matrix.h
@@ -6,6 +6,7 @@
 
 class FDMXPixelMappingToolkit;
 class IDetailLayoutBuilder;
+class IPropertyHandle;
 
 
 class FDMXPixelMappingDetailCustomization_Matrix
@@ -20,6 +21,12 @@ public:
 	virtual void CustomizeDetails(IDetailLayoutBuilder& DetailLayout) override;
 
 private:	
+	/** Returns true if the EditorColor property is editable */
+	bool CanEditEditorColor() const;
+
+	/** Handle for the bUsePatchColor property */
+	TSharedPtr<IPropertyHandle> UsePatchColorHandle;
+
 	/** Weak reference to the DMX editor */
 	TWeakPtr<FDMXPixelMappingToolkit> ToolkitWeakPtr;
 };
diff --git a/Engine/Plugins/VirtualProduction/DMX/DMXPixelMapping/Source/DMXPixelMappingEditor/Private/ViewModels/DMXPixelMappingDMXLibraryViewModel.cpp b/Engine/Plugins/VirtualProduction/DMX/DMXPixelMapping/Source/DMXPixelMappingEditor/Private/ViewModels/DMXPixelMappingDMXLibraryViewModel.cpp
--- a/Engine/Plugins/VirtualProduction/DMX/DMXPixelMapping/Source/DMXPixelMappingEditor/Private/ViewModels/DMXPixelMappingDMXLibraryViewModel.cpp
+++ b/Engine/Plugins/VirtualProduction/DMX/DMXPixelMapping/Source/DMXPixelMappingEditor/Private/ViewModels/DMXPixelMappingDMXLibraryViewModel.cpp
@@ -122,6 +122,21 @@ void UDMXPixelMappingDMXLibraryViewModel::UpdateFixtureGroupFromSelection(TWeakP
 	}
 }
 
+void UDMXPixelMappingDMXLibraryViewModel::SetNewComponentsUsePatchColor(bool bUsePatchColor)
+{
+	UDMXPixelMapping* PixelMapping = WeakToolkit.IsValid() ? WeakToolkit.Pin()->GetDMXPixelMapping() : nullptr;
+	if (PixelMapping)
+	{
+		PixelMapping->bNewComponentsUsePatchColor = bUsePatchColor;
+	}
+}
+
+bool UDMXPixelMappingDMXLibraryViewModel::ShouldNewComponentsUsePatchColor() const
+{
+	const UDMXPixelMapping* PixelMapping = WeakToolkit.IsValid() ? WeakToolkit.Pin()->GetDMXPixelMapping() : nullptr;
+	return PixelMapping ? PixelMapping->bNewComponentsUsePatchColor : false;
+}
+
 void UDMXPixelMappingDMXLibraryViewModel::AddFixturePatchesEnsured(const TArray<UDMXEntityFixturePatch*>& FixturePatches)
 {
 	const TSharedPtr<FDMXPixelMappingToolkit> Toolkit = WeakToolkit.Pin();
diff --git a/Engine/Plugins/VirtualProduction/DMX/DMXPixelMapping/Source/DMXPixelMappingEditor/Private/ViewModels/DMXPixelMappingDMXLibraryViewModel.h b/Engine/Plugins/VirtualProduction/DMX/DMXPixelMapping/Source/DMXPixelMappingEditor/Private/ViewModels/DMXPixelMappingDMXLibraryViewModel.h
--- a/Engine/Plugins/VirtualProduction/DMX/DMXPixelMapping/Source/DMXPixelMappingEditor/Private/ViewModels/DMXPixelMappingDMXLibraryViewModel.h
+++ b/Engine/Plugins/VirtualProduction/DMX/DMXPixelMapping/Source/DMXPixelMappingEditor/Private/ViewModels/DMXPixelMappingDMXLibraryViewModel.h
@@ -35,6 +35,12 @@ public:
 	/** Sets the fixture group for this model from what's currently selected. Note only single selection is supported as of 5.3 */
 	void UpdateFixtureGroupFromSelection(TWeakPtr<FDMXPixelMappingToolkit> InWeakToolkit);
 
+	/** Sets if new components should use the patch color */
+	void SetNewComponentsUsePatchColor(bool bUsePatchColor);
+
+	/** Returns true if new components should use the patch color */
+	bool ShouldNewComponentsUsePatchColor() const;
+
 	/** Adds specified fixture patches to the pixel mapping. Ensures all patches are of the same library as the one of the current group. */
 	void AddFixturePatchesEnsured(const TArray<UDMXEntityFixturePatch*>& FixturePatches);
 
diff --git a/Engine/Plugins/VirtualProduction/DMX/DMXPixelMapping/Source/DMXPixelMappingEditor/Private/Views/SDMXPixelMappingDMXLibraryView.cpp b/Engine/Plugins/VirtualProduction/DMX/DMXPixelMapping/Source/DMXPixelMappingEditor/Private/Views/SDMXPixelMappingDMXLibraryView.cpp
--- a/Engine/Plugins/VirtualProduction/DMX/DMXPixelMapping/Source/DMXPixelMappingEditor/Private/Views/SDMXPixelMappingDMXLibraryView.cpp
+++ b/Engine/Plugins/VirtualProduction/DMX/DMXPixelMapping/Source/DMXPixelMappingEditor/Private/Views/SDMXPixelMappingDMXLibraryView.cpp
@@ -23,7 +23,9 @@
 #include "ViewModels/DMXPixelMappingDMXLibraryViewModel.h"
 #include "Widgets/Images/SImage.h"
 #include "Widgets/Input/SButton.h"
+#include "Widgets/Input/SCheckBox.h"
 #include "Widgets/Layout/SWidgetSwitcher.h"
+#include "Widgets/Layout/SWrapBox.h"
 #include "Widgets/SBoxPanel.h"
 #include "Widgets/SDMXPixelMappingFixturePatchList.h"
 
@@ -133,15 +135,15 @@ void SDMXPixelMappingDMXLibraryView::Construct(const FArguments& InArgs, const T
 				ModelDetailsView.ToSharedRef()
 			]
 
-			// 'Add selected Fixture Patches' and 'Add all Patches' buttons
+			// Add Patches
 			+ SVerticalBox::Slot()
 			.AutoHeight()
 			[
-				SAssignNew(AddPatchesHorizontalBox, SHorizontalBox)
+				SAssignNew(AddPatchesWrapBox, SWrapBox)
+				.UseAllottedWidth(true)
 
-				+ SHorizontalBox::Slot()
+				+ SWrapBox::Slot()
 				.Padding(4.f)
-				.AutoWidth()
 				[
 					SNew(SButton)
 					.ButtonStyle(FAppStyle::Get(), "FlatButton.Success")
@@ -156,9 +158,8 @@ void SDMXPixelMappingDMXLibraryView::Construct(const FArguments& InArgs, const T
 					]
 				]
 
-				+ SHorizontalBox::Slot()
+				+ SWrapBox::Slot()
 				.Padding(8.f, 4.f)
-				.AutoWidth()
 				[
 					SNew(SButton)
 					.ButtonStyle(FAppStyle::Get(), "FlatButton.Success")
@@ -172,8 +173,23 @@ void SDMXPixelMappingDMXLibraryView::Construct(const FArguments& InArgs, const T
 						.Text(LOCTEXT("AddAllPatchesLabel", "Add All Patches"))
 					]
 				]
+
+				+ SWrapBox::Slot()
+				.Padding(8.f, 4.f)
+				.VAlign(VAlign_Center)
+				[
+					SNew(SCheckBox)
+					.IsChecked(this, &SDMXPixelMappingDMXLibraryView::GetUsePatchColorCheckState)
+					.OnCheckStateChanged(this, &SDMXPixelMappingDMXLibraryView::OnUsePatchColorCheckStateChanged)
+					[
+						SNew(STextBlock)
+						.Text(LOCTEXT("UsePatchColorCheckBoxLabel", "Use Patch Color"))
+						.Font(FAppStyle::GetFontStyle("PropertyWindow.NormalFont"))
+					]
+				]
 			]
-		
+
+			// The fixture patch list
 			+ SVerticalBox::Slot()
 			.FillHeight(1.f)
 			[
@@ -292,7 +308,7 @@ void SDMXPixelMappingDMXLibraryView::ForceRefresh()
 
 		// Only show the list and options to add or select when fixture patches are available
 		const EVisibility AddFixturePatchOptionsVisibility = FixturePatchesInDMXLibrary.Num() > HiddenFixturePatches.Num() ? EVisibility::Visible : EVisibility::Collapsed;
-		AddPatchesHorizontalBox->SetVisibility(AddFixturePatchOptionsVisibility);
+		AddPatchesWrapBox->SetVisibility(AddFixturePatchOptionsVisibility);
 
 		if (AddFixturePatchOptionsVisibility == EVisibility::Visible)
 		{
@@ -390,6 +406,20 @@ FReply SDMXPixelMappingDMXLibraryView::OnAddAllPatchesClicked()
 	return FReply::Handled();
 }
 
+ECheckBoxState SDMXPixelMappingDMXLibraryView::GetUsePatchColorCheckState() const
+{
+	return ViewModel && ViewModel->ShouldNewComponentsUsePatchColor() ? ECheckBoxState::Checked : ECheckBoxState::Unchecked;
+}
+
+void SDMXPixelMappingDMXLibraryView::OnUsePatchColorCheckStateChanged(ECheckBoxState NewCheckState)
+{
+	if (ensureMsgf(ViewModel, TEXT("Invalid view model for PixelMapping DMX Library View, cannot display view.")))
+	{
+		const bool bNewPatchesShouldUsePatchColor = NewCheckState == ECheckBoxState::Checked;
+		ViewModel->SetNewComponentsUsePatchColor(bNewPatchesShouldUsePatchColor);
+	}
+}
+
 TArray<UDMXEntityFixturePatch*> SDMXPixelMappingDMXLibraryView::GetFixturePatchesInDMXLibrary() const
 {
 	if (!ensureMsgf(ViewModel, TEXT("Invalid view model for PixelMapping DMX Library View, cannot display view.")))
diff --git a/Engine/Plugins/VirtualProduction/DMX/DMXPixelMapping/Source/DMXPixelMappingEditor/Private/Views/SDMXPixelMappingDMXLibraryView.h b/Engine/Plugins/VirtualProduction/DMX/DMXPixelMapping/Source/DMXPixelMappingEditor/Private/Views/SDMXPixelMappingDMXLibraryView.h
--- a/Engine/Plugins/VirtualProduction/DMX/DMXPixelMapping/Source/DMXPixelMappingEditor/Private/Views/SDMXPixelMappingDMXLibraryView.h
+++ b/Engine/Plugins/VirtualProduction/DMX/DMXPixelMapping/Source/DMXPixelMappingEditor/Private/Views/SDMXPixelMappingDMXLibraryView.h
@@ -8,14 +8,15 @@
 #include "Widgets/SCompoundWidget.h"
 #include "UObject/GCObject.h"
 
+enum class ECheckBoxState : uint8; 
 class FDMXPixelMappingToolkit;
 class FReply;
 class IDetailsView;
 class SBorder;
 class SDMXPixelMappingFixturePatchList;
-class SHorizontalBox;
 class STextBlock;
 class SWidgetSwitcher;
+class SWrapBox;
 class UDMXEntity;
 class UDMXLibrary;
 class UDMXPixelMapping;
@@ -79,6 +80,12 @@ private:
 	/** Called when the 'Add All Patches' button was clicked */
 	FReply OnAddAllPatchesClicked();
 
+	/** Returns the check state of the 'use patch color' check box */
+	ECheckBoxState GetUsePatchColorCheckState() const;
+
+	/** Called when the check state of the 'use patch color' check box changed */
+	void OnUsePatchColorCheckStateChanged(ECheckBoxState NewCheckState);
+
 	/** Helper that returns all fixture patches in use in the pixel mapping */
 	TArray<UDMXEntityFixturePatch*> GetFixturePatchesInDMXLibrary() const;
 
@@ -106,8 +113,8 @@ private:
 	/** Text block displaying 'all patches added' when no patches are available to add from the dmx library */
 	TSharedPtr<STextBlock> AllPatchesAddedTextBlock;
 
-	/** Horizontal box containing the 'add selected patches' and 'add all patches' buttons */
-	TSharedPtr<SHorizontalBox> AddPatchesHorizontalBox;
+	/** Wrap box containing the 'add selected patches' and 'add all patches' buttons and the 'Use Patch Color' checkbox */
+	TSharedPtr<SWrapBox> AddPatchesWrapBox;
 
 	/** Details view for the model, displays the library picker */
 	TSharedPtr<IDetailsView> ModelDetailsView;
diff --git a/Engine/Plugins/VirtualProduction/DMX/DMXPixelMapping/Source/DMXPixelMappingEditor/Private/Views/SDMXPixelMappingPreviewView.cpp b/Engine/Plugins/VirtualProduction/DMX/DMXPixelMapping/Source/DMXPixelMappingEditor/Private/Views/SDMXPixelMappingPreviewView.cpp
--- a/Engine/Plugins/VirtualProduction/DMX/DMXPixelMapping/Source/DMXPixelMappingEditor/Private/Views/SDMXPixelMappingPreviewView.cpp
+++ b/Engine/Plugins/VirtualProduction/DMX/DMXPixelMapping/Source/DMXPixelMappingEditor/Private/Views/SDMXPixelMappingPreviewView.cpp
@@ -291,7 +291,6 @@ TSharedRef<SWidget> SDMXPixelMappingPreviewView::CreateOverlayUI()
 				.ButtonColorAndOpacity(FLinearColor::Transparent)
 				.ToolTipText(LOCTEXT("ZoomToFit_ToolTip", "Zoom To Fit"))
 				.OnClicked(this, &SDMXPixelMappingPreviewView::HandleZoomToFitClicked)
-				.ContentPadding(FAppStyle::Get().GetMargin("ViewportMenu.SToolBarButtonBlock.Button.Padding"))
 				[
 					SNew(SImage)
 					.Image(FAppStyle::GetBrush("UMGEditor.ZoomToFit"))
diff --git a/Engine/Plugins/VirtualProduction/DMX/DMXPixelMapping/Source/DMXPixelMappingRuntime/Private/Components/DMXPixelMappingBaseComponent.cpp b/Engine/Plugins/VirtualProduction/DMX/DMXPixelMapping/Source/DMXPixelMappingRuntime/Private/Components/DMXPixelMappingBaseComponent.cpp
--- a/Engine/Plugins/VirtualProduction/DMX/DMXPixelMapping/Source/DMXPixelMappingRuntime/Private/Components/DMXPixelMappingBaseComponent.cpp
+++ b/Engine/Plugins/VirtualProduction/DMX/DMXPixelMapping/Source/DMXPixelMappingRuntime/Private/Components/DMXPixelMappingBaseComponent.cpp
@@ -100,7 +100,7 @@ void UDMXPixelMappingBaseComponent::ForEachChild(TComponentPredicate Predicate,
 	}
 }
 
-UDMXPixelMapping* UDMXPixelMappingBaseComponent::GetPixelMapping()
+UDMXPixelMapping* UDMXPixelMappingBaseComponent::GetPixelMapping() const
 {
 	if (const UDMXPixelMappingRootComponent* RootComponent = GetRootComponent())
 	{
diff --git a/Engine/Plugins/VirtualProduction/DMX/DMXPixelMapping/Source/DMXPixelMappingRuntime/Private/Components/DMXPixelMappingMatrixCellComponent.cpp b/Engine/Plugins/VirtualProduction/DMX/DMXPixelMapping/Source/DMXPixelMappingRuntime/Private/Components/DMXPixelMappingMatrixCellComponent.cpp
--- a/Engine/Plugins/VirtualProduction/DMX/DMXPixelMapping/Source/DMXPixelMappingRuntime/Private/Components/DMXPixelMappingMatrixCellComponent.cpp
+++ b/Engine/Plugins/VirtualProduction/DMX/DMXPixelMapping/Source/DMXPixelMappingRuntime/Private/Components/DMXPixelMappingMatrixCellComponent.cpp
@@ -23,7 +23,11 @@ UDMXPixelMappingMatrixCellComponent::UDMXPixelMappingMatrixCellComponent()
 {
 #if WITH_EDITORONLY_DATA
 	bLockInDesigner = true;
+
+	PRAGMA_DISABLE_DEPRECATION_WARNINGS
+	// EditorColor should no longer be accessed publicly, however it is ok to access it here.
 	EditorColor = FLinearColor::White.CopyWithNewOpacity(.25f);
+	PRAGMA_ENABLE_DEPRECATION_WARNINGS
 #endif // WITH_EDITORONLY_DATA
 }
 
@@ -72,13 +76,6 @@ void UDMXPixelMappingMatrixCellComponent::PostEditChangeProperty(FPropertyChange
 }
 #endif // WITH_EDITOR
 
-#if WITH_EDITOR
-FLinearColor UDMXPixelMappingMatrixCellComponent::GetEditorColor() const
-{
-	return EditorColor;
-}
-#endif // WITH_EDITOR
-
 bool UDMXPixelMappingMatrixCellComponent::IsOverParent() const
 {
 	if (UDMXPixelMappingMatrixComponent* Parent = Cast<UDMXPixelMappingMatrixComponent>(GetParent()))
diff --git a/Engine/Plugins/VirtualProduction/DMX/DMXPixelMapping/Source/DMXPixelMappingRuntime/Private/Components/DMXPixelMappingOutputComponent.cpp b/Engine/Plugins/VirtualProduction/DMX/DMXPixelMapping/Source/DMXPixelMappingRuntime/Private/Components/DMXPixelMappingOutputComponent.cpp
--- a/Engine/Plugins/VirtualProduction/DMX/DMXPixelMapping/Source/DMXPixelMappingRuntime/Private/Components/DMXPixelMappingOutputComponent.cpp
+++ b/Engine/Plugins/VirtualProduction/DMX/DMXPixelMapping/Source/DMXPixelMappingRuntime/Private/Components/DMXPixelMappingOutputComponent.cpp
@@ -92,12 +92,20 @@ bool UDMXPixelMappingOutputComponent::IsVisible() const
 #if WITH_EDITOR
 void UDMXPixelMappingOutputComponent::SetZOrder(int32 NewZOrder)
 {
-	const int32 DeltaZOrder = NewZOrder - ZOrder;
-	
 	ZOrder = NewZOrder;
 }
 #endif // WITH_EDITOR
 
+#if WITH_EDITOR
+FLinearColor UDMXPixelMappingOutputComponent::GetEditorColor() const
+{
+	PRAGMA_DISABLE_DEPRECATION_WARNINGS
+	// EditorColor should not be publicly accessed. However it is ok to access it here.
+	return EditorColor;
+	PRAGMA_ENABLE_DEPRECATION_WARNINGS
+}
+#endif // WITH_EDITOR
+
 bool UDMXPixelMappingOutputComponent::IsOverParent() const
 {
 	// By default all components are over their parent.
diff --git a/Engine/Plugins/VirtualProduction/DMX/DMXPixelMapping/Source/DMXPixelMappingRuntime/Private/Components/DMXPixelMappingOutputDMXComponent.cpp b/Engine/Plugins/VirtualProduction/DMX/DMXPixelMapping/Source/DMXPixelMappingRuntime/Private/Components/DMXPixelMappingOutputDMXComponent.cpp
new file mode 100644
--- /dev/null
+++ b/Engine/Plugins/VirtualProduction/DMX/DMXPixelMapping/Source/DMXPixelMappingRuntime/Private/Components/DMXPixelMappingOutputDMXComponent.cpp
@@ -0,0 +1,22 @@
+// Copyright Epic Games, Inc. All Rights Reserved.
+
+#include "Components/DMXPixelMappingOutputDMXComponent.h"
+
+#include "Library/DMXEntityFixturePatch.h"
+
+
+#if WITH_EDITOR
+FLinearColor UDMXPixelMappingOutputDMXComponent::GetEditorColor() const
+{
+	if (bUsePatchColor)
+	{
+		const UDMXEntityFixturePatch* FixturePatch = FixturePatchRef.GetFixturePatch();
+		return FixturePatch->EditorColor;
+	}
+
+	PRAGMA_DISABLE_DEPRECATION_WARNINGS
+	// Editor color should not be publicly accessed. However it is ok to access it here.
+	return EditorColor;
+	PRAGMA_ENABLE_DEPRECATION_WARNINGS
+}
+#endif // WITH_EDITOR
diff --git a/Engine/Plugins/VirtualProduction/DMX/DMXPixelMapping/Source/DMXPixelMappingRuntime/Private/Templates/DMXPixelMappingComponentTemplate.cpp b/Engine/Plugins/VirtualProduction/DMX/DMXPixelMapping/Source/DMXPixelMappingRuntime/Private/Templates/DMXPixelMappingComponentTemplate.cpp
--- a/Engine/Plugins/VirtualProduction/DMX/DMXPixelMapping/Source/DMXPixelMappingRuntime/Private/Templates/DMXPixelMappingComponentTemplate.cpp
+++ b/Engine/Plugins/VirtualProduction/DMX/DMXPixelMapping/Source/DMXPixelMappingRuntime/Private/Templates/DMXPixelMappingComponentTemplate.cpp
@@ -5,6 +5,7 @@
 #include "Components/DMXPixelMappingFixtureGroupItemComponent.h"
 #include "Components/DMXPixelMappingMatrixComponent.h"
 #include "Components/DMXPixelMappingRootComponent.h"
+#include "DMXPixelMapping.h"
 #include "Library/DMXEntityFixturePatch.h"
 
 
@@ -42,46 +43,58 @@ FText FDMXPixelMappingComponentTemplate::GetCategory() const
 
 UDMXPixelMappingBaseComponent* FDMXPixelMappingComponentTemplate::CreateComponentInternal(UDMXPixelMappingRootComponent* InRootComponent)
 {
-	if (ensureMsgf(InRootComponent, TEXT("Tried to create components from template but RootComponent was invalid.")))
+	const UDMXPixelMapping* PixelMapping = InRootComponent ? InRootComponent->GetPixelMapping() : nullptr;
+	if (!ensureMsgf(PixelMapping && InRootComponent, TEXT("Tried to create components from template but PixelMapping or RootComponent were invalid.")))
 	{
-		// Create a unique name
-		UDMXPixelMappingBaseComponent* DefaultComponent = ComponentClass->GetDefaultObject<UDMXPixelMappingBaseComponent>();
-		FName UniqueName = MakeUniqueObjectName(InRootComponent, ComponentClass.Get(), DefaultComponent->GetNamePrefix());
+		return nullptr;
+	}
 
-		if (ComponentClass.Get() == UDMXPixelMappingFixtureGroupItemComponent::StaticClass())
-		{
-			// Create the Component
-			UDMXPixelMappingFixtureGroupItemComponent* NewGroupItemComponent = NewObject<UDMXPixelMappingFixtureGroupItemComponent>(InRootComponent, ComponentClass.Get(), UniqueName, RF_Transactional | RF_Public);
+	// Create a unique name
+	UDMXPixelMappingBaseComponent* DefaultComponent = ComponentClass->GetDefaultObject<UDMXPixelMappingBaseComponent>();
+	FName UniqueName = MakeUniqueObjectName(InRootComponent, ComponentClass.Get(), DefaultComponent->GetNamePrefix());
 
-			NewGroupItemComponent->FixturePatchRef = FixturePatchRef;
+	UDMXPixelMappingBaseComponent* NewComponent = nullptr;
+	if (ComponentClass.Get() == UDMXPixelMappingFixtureGroupItemComponent::StaticClass())
+	{
+		// Create the Component
+		UDMXPixelMappingFixtureGroupItemComponent* NewGroupItemComponent = NewObject<UDMXPixelMappingFixtureGroupItemComponent>(InRootComponent, ComponentClass.Get(), UniqueName, RF_Transactional | RF_Public);
 
-			// Create a better unique name and set it
-			UniqueName = MakeUniqueObjectName(NewGroupItemComponent->GetOuter(), NewGroupItemComponent->GetClass(), FName(FixturePatchRef.GetFixturePatch()->Name));
-			const FString NewNameStr = UniqueName.ToString();
-			NewGroupItemComponent->Rename(*NewNameStr);
+		NewGroupItemComponent->FixturePatchRef = FixturePatchRef;
 
-			return NewGroupItemComponent;
-		}
-		else if (ComponentClass.Get() == UDMXPixelMappingMatrixComponent::StaticClass())
-		{
-			UDMXPixelMappingMatrixComponent* NewMatrixComponent = NewObject<UDMXPixelMappingMatrixComponent>(InRootComponent, ComponentClass.Get(), UniqueName, RF_Transactional | RF_Public);
+		NewComponent = NewGroupItemComponent;
+	}
+	else if (ComponentClass.Get() == UDMXPixelMappingMatrixComponent::StaticClass())
+	{
+		UDMXPixelMappingMatrixComponent* NewMatrixComponent = NewObject<UDMXPixelMappingMatrixComponent>(InRootComponent, ComponentClass.Get(), UniqueName, RF_Transactional | RF_Public);
+		NewMatrixComponent->FixturePatchRef = FixturePatchRef;
 
-			NewMatrixComponent->FixturePatchRef = FixturePatchRef;
+		NewComponent = NewMatrixComponent;
+	}
+	else
+	{
+		UDMXPixelMappingBaseComponent* NewBaseComponent = NewObject<UDMXPixelMappingBaseComponent>(InRootComponent, ComponentClass.Get(), UniqueName, RF_Transactional | RF_Public);
 
-			// Create a better unique name and set it
-			UniqueName = MakeUniqueObjectName(NewMatrixComponent->GetOuter(), NewMatrixComponent->GetClass(), FName(FixturePatchRef.GetFixturePatch()->Name));
-			const FString NewNameStr = UniqueName.ToString();
-			NewMatrixComponent->Rename(*NewNameStr);
+		NewComponent = NewBaseComponent;
+	}
+	check(NewComponent);
 
-			return NewMatrixComponent;
-		}
-		else
+	// Initialize name and color of components that use a patch with data from the patch
+	if (UDMXPixelMappingOutputDMXComponent* NewComponentWithPatch = Cast<UDMXPixelMappingOutputDMXComponent>(NewComponent))
+	{
+		const UDMXEntityFixturePatch* Patch = NewComponentWithPatch->FixturePatchRef.GetFixturePatch();
+		if (Patch)
 		{
-			UDMXPixelMappingBaseComponent* NewComponent = NewObject<UDMXPixelMappingBaseComponent>(InRootComponent, ComponentClass.Get(), UniqueName, RF_Transactional | RF_Public);
+			// Create a better unique name and set it
+			UniqueName = MakeUniqueObjectName(NewComponentWithPatch->GetOuter(), NewComponentWithPatch->GetClass(), *Patch->Name);
+			const FString NewNameStr = UniqueName.ToString();
+			NewComponentWithPatch->Rename(*NewNameStr);
 
-			return NewComponent;
+#if WITH_EDITOR
+			// Set if the component should follow the patch color
+			NewComponentWithPatch->bUsePatchColor = PixelMapping->bNewComponentsUsePatchColor;
+#endif // WITH_EDITOR
 		}
 	}
 
-	return nullptr;
+	return NewComponent;
 }
diff --git a/Engine/Plugins/VirtualProduction/DMX/DMXPixelMapping/Source/DMXPixelMappingRuntime/Public/Components/DMXPixelMappingBaseComponent.h b/Engine/Plugins/VirtualProduction/DMX/DMXPixelMapping/Source/DMXPixelMappingRuntime/Public/Components/DMXPixelMappingBaseComponent.h
--- a/Engine/Plugins/VirtualProduction/DMX/DMXPixelMapping/Source/DMXPixelMappingRuntime/Public/Components/DMXPixelMappingBaseComponent.h
+++ b/Engine/Plugins/VirtualProduction/DMX/DMXPixelMapping/Source/DMXPixelMappingRuntime/Public/Components/DMXPixelMappingBaseComponent.h
@@ -150,7 +150,7 @@ public:
 	}
 
 	/** Get Pixel Mapping asset UObject */
-	UDMXPixelMapping* GetPixelMapping();
+	UDMXPixelMapping* GetPixelMapping() const;
 
 	/** Get root component of the component tree */
 	const UDMXPixelMappingRootComponent* GetRootComponent() const;
diff --git a/Engine/Plugins/VirtualProduction/DMX/DMXPixelMapping/Source/DMXPixelMappingRuntime/Public/Components/DMXPixelMappingMatrixCellComponent.h b/Engine/Plugins/VirtualProduction/DMX/DMXPixelMapping/Source/DMXPixelMappingRuntime/Public/Components/DMXPixelMappingMatrixCellComponent.h
--- a/Engine/Plugins/VirtualProduction/DMX/DMXPixelMapping/Source/DMXPixelMappingRuntime/Public/Components/DMXPixelMappingMatrixCellComponent.h
+++ b/Engine/Plugins/VirtualProduction/DMX/DMXPixelMapping/Source/DMXPixelMappingRuntime/Public/Components/DMXPixelMappingMatrixCellComponent.h
@@ -49,10 +49,7 @@ public:
 	virtual FString GetUserName() const override;
 	//~ End UDMXPixelMappingBaseComponent implementation
 
-	//~ Begin UDMXPixelMappingOutputComponent implementation
-#if WITH_EDITOR
-	virtual FLinearColor GetEditorColor() const override;
-#endif // WITH_EDITOR	
+	//~ Begin UDMXPixelMappingOutputComponent implementation	
 	virtual bool IsOverParent() const override;
 	virtual int32 GetDownsamplePixelIndex() const override { return DownsamplePixelIndex; }
 	virtual void QueueDownsample() override;
diff --git a/Engine/Plugins/VirtualProduction/DMX/DMXPixelMapping/Source/DMXPixelMappingRuntime/Public/Components/DMXPixelMappingMatrixComponent.h b/Engine/Plugins/VirtualProduction/DMX/DMXPixelMapping/Source/DMXPixelMappingRuntime/Public/Components/DMXPixelMappingMatrixComponent.h
--- a/Engine/Plugins/VirtualProduction/DMX/DMXPixelMapping/Source/DMXPixelMappingRuntime/Public/Components/DMXPixelMappingMatrixComponent.h
+++ b/Engine/Plugins/VirtualProduction/DMX/DMXPixelMapping/Source/DMXPixelMappingRuntime/Public/Components/DMXPixelMappingMatrixComponent.h
@@ -115,11 +115,11 @@ public:
 	UPROPERTY()
 	EDMXPixelMappingDistribution Distribution;
 
-	/** Inverts the cell ordering along the X-axis. Useful when the matrix is facing away from the screen, rotated horizontally. */
+	/** Inverts the cell ordering along the X-axis. Useful when the matrix is facing away from the screen. */
 	UPROPERTY(EditAnywhere, Category = "Matrix")
 	bool bInvertCellsX = false;
 
-	/** Inverts the cell ordering along the Y-axis. Useful when the matrix is facing away from the screen, rotated vertically. */
+	/** Inverts the cell ordering along the Y-axis. Useful when the matrix is facing away from the screen . */
 	UPROPERTY(EditAnywhere, Category = "Matrix")
 	bool bInvertCellsY = false;
 
diff --git a/Engine/Plugins/VirtualProduction/DMX/DMXPixelMapping/Source/DMXPixelMappingRuntime/Public/Components/DMXPixelMappingOutputComponent.h b/Engine/Plugins/VirtualProduction/DMX/DMXPixelMapping/Source/DMXPixelMappingRuntime/Public/Components/DMXPixelMappingOutputComponent.h
--- a/Engine/Plugins/VirtualProduction/DMX/DMXPixelMapping/Source/DMXPixelMappingRuntime/Public/Components/DMXPixelMappingOutputComponent.h
+++ b/Engine/Plugins/VirtualProduction/DMX/DMXPixelMapping/Source/DMXPixelMappingRuntime/Public/Components/DMXPixelMappingOutputComponent.h
@@ -65,7 +65,7 @@ public:
 	virtual int32 GetZOrder() const { return ZOrder; }
 
 	/** Returns an editor color for the widget */
-	virtual FLinearColor GetEditorColor() const { return EditorColor; }
+	virtual FLinearColor GetEditorColor() const;
 #endif // WITH_EDITOR
 
 	/** Returns true if the the component's over all its parents. */
@@ -76,15 +76,15 @@ public:
 
 	/** Returns true if the component overlaps the other */
 	UE_DEPRECATED(5.4, "Removed without replacement. GetEdges and test against other if this method is needed.")
-		virtual bool OverlapsComponent(UDMXPixelMappingOutputComponent* Other) const;
+	virtual bool OverlapsComponent(UDMXPixelMappingOutputComponent* Other) const;
 
 	/** Get pixel index in downsample texture */
 	UE_DEPRECATED(5.3, "Please use UDMXPixelMappingPixelMapRenderer to render the pixel map")
-		virtual int32 GetDownsamplePixelIndex() const { return 0; }
+	virtual int32 GetDownsamplePixelIndex() const { return 0; }
 
 	/** Queue rendering to downsample rendering target */
 	UE_DEPRECATED(5.3, "Please use UDMXPixelMappingPixelMapRenderer to render the pixel map")
-		virtual void QueueDownsample() {}
+	virtual void QueueDownsample() {}
 
 	/** Returns the absolute position, without rotation */
 	FVector2D GetPosition() const;
@@ -132,11 +132,11 @@ public:
 	int32 ZOrder = 1;
 
 	/** The color displayed in editor */
+	UE_DEPRECATED(5.4, "Should no longer be publicly accessed, instead call GetEditorColor.")
 	UPROPERTY(BlueprintReadWrite, EditAnywhere, Category = "Editor Settings")
 	FLinearColor EditorColor = FLinearColor::Blue;
 #endif // WITH_EDITORONLY_DATA
 
-public:
 #if WITH_EDITOR
 	// Property Name getters
 	FORCEINLINE static FName GetLockInDesignerPropertyName() { return GET_MEMBER_NAME_CHECKED(UDMXPixelMappingOutputComponent, bLockInDesigner); }
@@ -156,24 +156,23 @@ protected:
 	float PositionY = 0.f;
 
 #if WITH_EDITORONLY_DATA
-	UPROPERTY(EditAnywhere, Transient, Category = "Transform", Meta = (EditCondition = "!bLockInDesigner", AllowPrivateAccess = true, DisplayName = "Position (with Rotation)"))
+	UPROPERTY(EditAnywhere, Transient, Category = "Transform", Meta = (EditCondition = "!bLockInDesigner", DisplayName = "Position (with Rotation)"))
 	FVector2D EditorPositionWithRotation = FVector2D::ZeroVector;
 #endif // WITH_EDITORONLY_DATA
 
-	UPROPERTY(EditAnywhere, BlueprintReadOnly, Category = "Transform", Meta = (ClampMin = 0.0001, EditCondition = "!bLockInDesigner", AllowPrivateAccess = true))
+	UPROPERTY(EditAnywhere, BlueprintReadOnly, Category = "Transform", Meta = (ClampMin = 0.0001, EditCondition = "!bLockInDesigner"))
 	float SizeX = 1.f;
 
-	UPROPERTY(EditAnywhere, BlueprintReadOnly, Category = "Transform", Meta = (ClampMin = 0.0001, EditCondition = "!bLockInDesigner", AllowPrivateAccess = true))
+	UPROPERTY(EditAnywhere, BlueprintReadOnly, Category = "Transform", Meta = (ClampMin = 0.0001, EditCondition = "!bLockInDesigner"))
 	float SizeY = 1.f;
 
-	UPROPERTY(EditAnywhere, BlueprintReadOnly, Category = "Transform", Meta = (UIMin = -180, UIMax = 180, EditCondition = "!bLockInDesigner", AllowPrivateAccess = true))
+	UPROPERTY(EditAnywhere, BlueprintReadOnly, Category = "Transform", Meta = (UIMin = -180, UIMax = 180, EditCondition = "!bLockInDesigner"))
 	double Rotation = 0.0;
 
 	UPROPERTY(Transient)
 	FDMXPixelMappingComponentGeometryCache CachedGeometry;
 
 #if WITH_EDITORONLY_DATA
-protected:
 	UPROPERTY(EditAnywhere, BlueprintReadOnly, Category = "Editor Settings")
 	bool bLockInDesigner;
 
diff --git a/Engine/Plugins/VirtualProduction/DMX/DMXPixelMapping/Source/DMXPixelMappingRuntime/Public/Components/DMXPixelMappingOutputDMXComponent.h b/Engine/Plugins/VirtualProduction/DMX/DMXPixelMapping/Source/DMXPixelMappingRuntime/Public/Components/DMXPixelMappingOutputDMXComponent.h
--- a/Engine/Plugins/VirtualProduction/DMX/DMXPixelMapping/Source/DMXPixelMappingRuntime/Public/Components/DMXPixelMappingOutputDMXComponent.h
+++ b/Engine/Plugins/VirtualProduction/DMX/DMXPixelMapping/Source/DMXPixelMappingRuntime/Public/Components/DMXPixelMappingOutputDMXComponent.h
@@ -24,6 +24,12 @@ class DMXPIXELMAPPINGRUNTIME_API UDMXPixelMappingOutputDMXComponent
 	GENERATED_BODY()
 
 public:
+	//~ Begin UDMXPixelMappingOutputComponent interface
+#if WITH_EDITOR
+	virtual FLinearColor GetEditorColor() const override;
+#endif // WITH_EDITOR	
+	//~ End UDMXPixelMappingOutputComponent interface
+
 	/** The quality level to use when averaging colors during downsampling. */
 	UPROPERTY(EditAnywhere, BlueprintReadOnly, Category = "Quality")
 	EDMXPixelBlendingQuality CellBlendingQuality;
@@ -31,6 +37,12 @@ public:
 	UPROPERTY(VisibleAnywhere, BlueprintReadOnly, Category = "Fixture Patch", Meta = (ShowOnlyInnerProperties))
 	FDMXEntityFixturePatchRef FixturePatchRef;
 
+#if WITH_EDITORONLY_DATA
+	/** If set, the color of the patch is displayed, instead of a custom editor color */
+	UPROPERTY(EditAnywhere, BlueprintReadOnly, Category = "Editor Settings", Meta = (AllowPrivateAccess = true))
+	bool bUsePatchColor = false;
+#endif
+
 	/** Sets which color space Pixel Mapping sends */
 	UPROPERTY(Transient, EditAnywhere, NoClear, Category = "Color Space", Meta = (DisplayPriority = 2, DisplayName = "Output Mode", ShowDisplayNames))
 	TSubclassOf<UDMXPixelMappingColorSpace> ColorSpaceClass;
diff --git a/Engine/Plugins/VirtualProduction/DMX/DMXPixelMapping/Source/DMXPixelMappingRuntime/Public/DMXPixelMapping.h b/Engine/Plugins/VirtualProduction/DMX/DMXPixelMapping/Source/DMXPixelMappingRuntime/Public/DMXPixelMapping.h
--- a/Engine/Plugins/VirtualProduction/DMX/DMXPixelMapping/Source/DMXPixelMappingRuntime/Public/DMXPixelMapping.h
+++ b/Engine/Plugins/VirtualProduction/DMX/DMXPixelMapping/Source/DMXPixelMappingRuntime/Public/DMXPixelMapping.h
@@ -164,6 +164,10 @@ public:
 	UPROPERTY(NonTransactional)
 	float ComponentLabelFontSize = 8.f;
 
+	/** If true, new components use the fixture patch color instead of the default pixel mapping color. */
+	UPROPERTY(NonTransactional)
+	bool bNewComponentsUsePatchColor = true;
+
 	/** If true, editor is set to scale children with parent. This is forwarded from the editor module (DMXPixelMappingEditorSettings) to be accessible in the runtime module. */
 	UPROPERTY(Transient, NonTransactional)
 	bool bEditorScaleChildrenWithParent = false;
