DMX - Add PixelMapping HDR capabilities, rendering considers rotation

diff --git a/Engine/Plugins/VirtualProduction/DMX/DMXPixelMapping/Shaders/Private/DMXPixelMapping.usf b/Engine/Plugins/VirtualProduction/DMX/DMXPixelMapping/Shaders/Private/DMXPixelMapping.usf
--- a/Engine/Plugins/VirtualProduction/DMX/DMXPixelMapping/Shaders/Private/DMXPixelMapping.usf
+++ b/Engine/Plugins/VirtualProduction/DMX/DMXPixelMapping/Shaders/Private/DMXPixelMapping.usf
@@ -47,42 +47,53 @@ OutputPS DMXPixelMappingPS(OutputVS IN)
 	}
 	else
 	{
-		// Average the final color sample based on various samples taken from within this UV space.
+		// Sample the center color
 		OUT.Color = Texture2DSample(InputTexture, InputSampler, UV);
 
+		// For higher qualities, average the final color from the center color and additional samples on the sides.
 #if PIXELBLENDING_QUALITY == PIXELBLENDING_QUALITY_MEDIUM
 		#define NUM_CELLS 4
 		float2 UVCells[NUM_CELLS];
 		
-		UVCells[0] = UV + float2(-UVCellSize.x, -UVCellSize.y);			// UL
-		UVCells[1] = UV + float2(UVCellSize.x, -UVCellSize.y);			// UR
-		UVCells[2] = UV + float2(UVCellSize.x, UVCellSize.y);			// BR
-		UVCells[3] = UV + float2(-UVCellSize.x, UVCellSize.y);			// BL
+		// 1/8 from the vertices results in an even color distribution with 5 samples
+		UVCells[0] = UVTopLeftRotated + (UV - UVTopLeftRotated) / 8.f;			// Upper Left 
+		UVCells[1] = UVTopRightRotated + (UV - UVTopRightRotated) / 8.f;		// Upper Right
+		UVCells[2] = UV + UV - UVCells[0];										// Bottom Right
+		UVCells[3] = UV + UV - UVCells[1];										// Bottom Left
+
+		for(uint i = 0; i < NUM_CELLS; ++i)
+		{
+			// Only add 1/4 of the color to avoid bias towards the sides
+			OUT.Color += Texture2DSample(InputTexture, InputSampler, clamp(UVCells[i], 0, 1)) / 4.f;
+		}
+		
+		// The average of the center color + 1/4 of the color on each side
+		OUT.Color /= 2.f;
 
 #elif PIXELBLENDING_QUALITY == PIXELBLENDING_QUALITY_HIGH
 		#define NUM_CELLS 8
 		float2 UVCells[NUM_CELLS];
 		
-		UVCells[0] = UV + float2(-UVCellSize.x, -UVCellSize.y);			// UL
-		UVCells[1] = UV + float2(UVCellSize.x, -UVCellSize.y);			// UR
-		UVCells[2] = UV + float2(UVCellSize.x, UVCellSize.y);			// BR
-		UVCells[3] = UV + float2(-UVCellSize.x, UVCellSize.y);			// BL
-		UVCells[4] = UV + float2(-UVCellSize.x, 0);						// L
-		UVCells[5] = UV + float2(0, -UVCellSize.y);						// U
-		UVCells[6] = UV + float2(UVCellSize.x, 0);						// R
-		UVCells[7] = UV + float2(0, UVCellSize.y);						// D
+		// 1/6 from the vertices results in an even color distribution with 9 samples
+		UVCells[0] = UVTopLeftRotated + (UV - UVTopLeftRotated) / 6.f;			// Upper Left  
+		UVCells[1] = UVTopRightRotated + (UV - UVTopRightRotated) / 6.f;		// Upper Right 
+		UVCells[2] = UV + UV - UVCells[0];										// Bottom Right
+		UVCells[3] = UV + UV - UVCells[1];										// Bottom Left
+		UVCells[4] = UV + (UVCells[0] - UV + UVCells[1] - UV) / 2.f;			// Upper Center
+		UVCells[5] = UV + (UVCells[1] - UV + UVCells[2] - UV) / 2.f;			// Right Center
+		UVCells[6] = UV + (UVCells[2] - UV + UVCells[3] - UV) / 2.f;			// Bottom Center
+		UVCells[7] = UV + (UVCells[3] - UV + UVCells[0] - UV) / 2.f;			// Left Center
 		
-#elif PIXELBLENDING_QUALITY > PIXELBLENDING_QUALITY_HIGH
-#error Invalid pixel blending quality level specified.
-#endif
-
-#if PIXELBLENDING_QUALITY > PIXELBLENDING_QUALITY_LOW
 		for(uint i = 0; i < NUM_CELLS; ++i)
 		{
 			OUT.Color += Texture2DSample(InputTexture, InputSampler, clamp(UVCells[i], 0, 1));
 		}
 		
-		OUT.Color /= (NUM_CELLS + 1);
+		// The average of the center color + all the points on sides
+		OUT.Color /= 9.f;
+		
+#elif PIXELBLENDING_QUALITY > PIXELBLENDING_QUALITY_HIGH
+#error Invalid pixel blending quality level specified.
 #endif
 
 		OUT.Color *= Brightness;
diff --git a/Engine/Plugins/VirtualProduction/DMX/DMXPixelMapping/Shaders/Private/DMXPixelMapping.ush b/Engine/Plugins/VirtualProduction/DMX/DMXPixelMapping/Shaders/Private/DMXPixelMapping.ush
--- a/Engine/Plugins/VirtualProduction/DMX/DMXPixelMapping/Shaders/Private/DMXPixelMapping.ush
+++ b/Engine/Plugins/VirtualProduction/DMX/DMXPixelMapping/Shaders/Private/DMXPixelMapping.ush
@@ -18,9 +18,8 @@ SamplerState InputSampler;
 
 // PS Parameters
 float Brightness;
-int2 InputTextureSize;
-int2 OutputTextureSize;
-float2 UVCellSize;
+float2 UVTopLeftRotated;
+float2 UVTopRightRotated;
 
 struct InputVS
 {
diff --git a/Engine/Plugins/VirtualProduction/DMX/DMXPixelMapping/Source/DMXPixelMappingEditor/Private/Viewports/DMXPixelMappingPreviewViewportClient.cpp b/Engine/Plugins/VirtualProduction/DMX/DMXPixelMapping/Source/DMXPixelMappingEditor/Private/Viewports/DMXPixelMappingPreviewViewportClient.cpp
--- a/Engine/Plugins/VirtualProduction/DMX/DMXPixelMapping/Source/DMXPixelMappingEditor/Private/Viewports/DMXPixelMappingPreviewViewportClient.cpp
+++ b/Engine/Plugins/VirtualProduction/DMX/DMXPixelMapping/Source/DMXPixelMappingEditor/Private/Viewports/DMXPixelMappingPreviewViewportClient.cpp
@@ -129,6 +129,8 @@ void FDMXPixelMappingPreviewViewportClient::Draw(FViewport* InViewport, FCanvas*
 
 				FCanvasTileItem TileItem(Position, Size, Element->GetColor());
 				TileItem.BlendMode = ESimpleElementBlendMode::SE_BLEND_MAX;
+				TileItem.PivotPoint = FVector2D(0.5, 0.5);
+				TileItem.Rotation = FRotator(0.0, Element->GetParameters().Rotation, 0.0);
 				Canvas->DrawItem(TileItem);
 			}
 		}
@@ -139,6 +141,8 @@ void FDMXPixelMappingPreviewViewportClient::Draw(FViewport* InViewport, FCanvas*
 
 			FCanvasTileItem TileItem(Position, Size, Element->GetColor());
 			TileItem.BlendMode = ESimpleElementBlendMode::SE_BLEND_MAX;
+			TileItem.PivotPoint = FVector2D(0.5, 0.5);
+			TileItem.Rotation = FRotator(0.0, Element->GetParameters().Rotation, 0.0);
 			Canvas->DrawItem(TileItem);
 		}
 	}
diff --git a/Engine/Plugins/VirtualProduction/DMX/DMXPixelMapping/Source/DMXPixelMappingEditor/Private/Viewports/DMXPixelMappingSourceTextureViewportClient.cpp b/Engine/Plugins/VirtualProduction/DMX/DMXPixelMapping/Source/DMXPixelMappingEditor/Private/Viewports/DMXPixelMappingSourceTextureViewportClient.cpp
--- a/Engine/Plugins/VirtualProduction/DMX/DMXPixelMapping/Source/DMXPixelMappingEditor/Private/Viewports/DMXPixelMappingSourceTextureViewportClient.cpp
+++ b/Engine/Plugins/VirtualProduction/DMX/DMXPixelMapping/Source/DMXPixelMappingEditor/Private/Viewports/DMXPixelMappingSourceTextureViewportClient.cpp
@@ -112,6 +112,7 @@ void FDMXPixelMappingSourceTextureViewportClient::Draw(FViewport* InViewport, FC
 	const FVector2D UV0 = VisibleRect.Min / TextureBox.GetSize();
 	const FVector2D UV1 = VisibleRect.Max / TextureBox.GetSize();
 
+	const FLinearColor ColorWithExposure = FLinearColor::White * PixelMapping->DesignerExposure;
 	if (IsDrawingVisibleRectOnly())
 	{
 		const TSharedPtr<FSceneViewport> SceneViewport = ViewportWidget->GetViewport();
@@ -120,13 +121,13 @@ void FDMXPixelMappingSourceTextureViewportClient::Draw(FViewport* InViewport, FC
 			return;
 		}
 
-		FCanvasTileItem TileItem(FVector2D::ZeroVector, InputTexture->GetResource(), SceneViewport->GetSizeXY(), UV0, UV1, FLinearColor::White);
+		FCanvasTileItem TileItem(FVector2D::ZeroVector, InputTexture->GetResource(), SceneViewport->GetSizeXY(), UV0, UV1, ColorWithExposure);
 		TileItem.BlendMode = ESimpleElementBlendMode::SE_BLEND_MAX;
 		Canvas->DrawItem(TileItem);
 	}
 	else
 	{
-		FCanvasTileItem TileItem(VisibleRect.Min * DesignerView->GetZoomAmount(), InputTexture->GetResource(), VisibleRect.GetSize() * DesignerView->GetZoomAmount(), UV0, UV1, FLinearColor::White);
+		FCanvasTileItem TileItem(VisibleRect.Min * DesignerView->GetZoomAmount(), InputTexture->GetResource(), VisibleRect.GetSize() * DesignerView->GetZoomAmount(), UV0, UV1, ColorWithExposure);
 		TileItem.BlendMode = ESimpleElementBlendMode::SE_BLEND_MAX;
 		Canvas->DrawItem(TileItem);
 	}
diff --git a/Engine/Plugins/VirtualProduction/DMX/DMXPixelMapping/Source/DMXPixelMappingEditor/Private/Widgets/SDMXPixelMappingDesignerToolbar.cpp b/Engine/Plugins/VirtualProduction/DMX/DMXPixelMapping/Source/DMXPixelMappingEditor/Private/Widgets/SDMXPixelMappingDesignerToolbar.cpp
--- a/Engine/Plugins/VirtualProduction/DMX/DMXPixelMapping/Source/DMXPixelMappingEditor/Private/Widgets/SDMXPixelMappingDesignerToolbar.cpp
+++ b/Engine/Plugins/VirtualProduction/DMX/DMXPixelMapping/Source/DMXPixelMappingEditor/Private/Widgets/SDMXPixelMappingDesignerToolbar.cpp
@@ -279,11 +279,20 @@ namespace UE::DMX
 			MenuBuilder.AddMenuEntry(FDMXPixelMappingEditorCommands::Get().ToggleShowCellIDs);
 			MenuBuilder.AddMenuEntry(FDMXPixelMappingEditorCommands::Get().ToggleShowPivot);
 
-			constexpr bool NoIndent = true;
+			constexpr bool bNoIndent = true;
 			MenuBuilder.AddWidget(
 				GenerateComponentFontSizeEditWidget(),
 				LOCTEXT("FontSizeLabel", "Font Size"),
-				NoIndent
+				bNoIndent
+			);
+
+			constexpr bool bSearchable = false;
+			MenuBuilder.AddWidget(
+				GenerateDesignerExposureEditWidget(),
+				LOCTEXT("ExposureLabel", "Display Exposure"),
+				bNoIndent,
+				bSearchable,
+				LOCTEXT("ExposureTooltip", "Adjusts the exposure of the Designer.\nHas no effect on the Pixel Mapping rendering and the DMX output.")
 			);
 		}
 		MenuBuilder.EndSection();
@@ -298,25 +307,19 @@ namespace UE::DMX
 
 	TSharedRef<SWidget> SDMXPixelMappingDesignerToolbar::GenerateComponentFontSizeEditWidget()
 	{
-		UDMXPixelMapping* PixelMapping = WeakToolkit.IsValid() ? WeakToolkit.Pin()->GetDMXPixelMapping() : nullptr;
-		if (PixelMapping)
-		{
-			return
-				SNew(SBorder)
-				.BorderImage(FAppStyle::GetBrush("NoBorder"))
-				.HAlign(HAlign_Right)
-				[
-					SNew(SNumericEntryBox<uint8>)
-					.MinDesiredValueWidth(40.f)
-					.MinSliderValue(6)
-					.MaxSliderValue(24)
-					.AllowSpin(true)
-					.Value(this, &SDMXPixelMappingDesignerToolbar::GetComponentFontSize)
-					.OnValueChanged(this, &SDMXPixelMappingDesignerToolbar::SetComponentFontSize)
-				];
-		}
-
-		return SNullWidget::NullWidget;
+		return
+			SNew(SBorder)
+			.BorderImage(FAppStyle::GetBrush("NoBorder"))
+			.HAlign(HAlign_Right)
+			[
+				SNew(SNumericEntryBox<uint8>)
+				.MinDesiredValueWidth(40.f)
+				.MinSliderValue(6)
+				.MaxSliderValue(24)
+				.AllowSpin(true)
+				.Value(this, &SDMXPixelMappingDesignerToolbar::GetComponentFontSize)
+				.OnValueChanged(this, &SDMXPixelMappingDesignerToolbar::SetComponentFontSize)
+			];
 	}
 
 	TOptional<uint8> SDMXPixelMappingDesignerToolbar::GetComponentFontSize() const
@@ -339,6 +342,44 @@ namespace UE::DMX
 			PixelMapping->ComponentLabelFontSize = FMath::Max(FontSize, MinFontSize);
 		}
 	}
+
+	TSharedRef<SWidget> SDMXPixelMappingDesignerToolbar::GenerateDesignerExposureEditWidget()
+	{
+		return
+			SNew(SBorder)
+			.BorderImage(FAppStyle::GetBrush("NoBorder"))
+			.HAlign(HAlign_Right)
+			[
+				SNew(SNumericEntryBox<float>)
+				.MinDesiredValueWidth(40.f)
+				.MinSliderValue(0.f)
+				.MaxSliderValue(1.f)
+				.AllowSpin(true)
+				.Value(this, &SDMXPixelMappingDesignerToolbar::GetDesignerExposure)
+				.OnValueChanged(this, &SDMXPixelMappingDesignerToolbar::SetDesignerExposure)
+			];
+	}
+
+	TOptional<float> SDMXPixelMappingDesignerToolbar::GetDesignerExposure() const
+	{
+		UDMXPixelMapping* PixelMapping = WeakToolkit.IsValid() ? WeakToolkit.Pin()->GetDMXPixelMapping() : nullptr;
+		if (PixelMapping)
+		{
+			return PixelMapping->DesignerExposure;
+		}
+
+		return TOptional<float>();
+	}
+
+	void SDMXPixelMappingDesignerToolbar::SetDesignerExposure(float Exposure)
+	{
+		UDMXPixelMapping* PixelMapping = WeakToolkit.IsValid() ? WeakToolkit.Pin()->GetDMXPixelMapping() : nullptr;
+		if (PixelMapping)
+		{
+			constexpr uint8 MinExposure = 0.f;
+			PixelMapping->DesignerExposure = FMath::Max(MinExposure, Exposure);
+		}
+	}
 }
 
 #undef LOCTEXT_NAMESPACE
diff --git a/Engine/Plugins/VirtualProduction/DMX/DMXPixelMapping/Source/DMXPixelMappingEditor/Private/Widgets/SDMXPixelMappingDesignerToolbar.h b/Engine/Plugins/VirtualProduction/DMX/DMXPixelMapping/Source/DMXPixelMappingEditor/Private/Widgets/SDMXPixelMappingDesignerToolbar.h
--- a/Engine/Plugins/VirtualProduction/DMX/DMXPixelMapping/Source/DMXPixelMappingEditor/Private/Widgets/SDMXPixelMappingDesignerToolbar.h
+++ b/Engine/Plugins/VirtualProduction/DMX/DMXPixelMapping/Source/DMXPixelMappingEditor/Private/Widgets/SDMXPixelMappingDesignerToolbar.h
@@ -65,6 +65,15 @@ namespace UE::DMX
 		/** Sets the font size for component labels in the designer view */
 		void SetComponentFontSize(uint8 FontSize);
 
+		/** Generates a widget to edit the display exposure */
+		TSharedRef<SWidget> GenerateDesignerExposureEditWidget();
+
+		/** Gets the exposure of the designer view */
+		TOptional<float> GetDesignerExposure() const;
+
+		/** Sets the exposure of the designer view */
+		void SetDesignerExposure(float Exposure);
+
 		/** The toolkit that owns this widget */
 		TWeakPtr<FDMXPixelMappingToolkit> WeakToolkit;
 	};
diff --git a/Engine/Plugins/VirtualProduction/DMX/DMXPixelMapping/Source/DMXPixelMappingRenderer/Private/DMXPixelMappingPixelMapRenderer.cpp b/Engine/Plugins/VirtualProduction/DMX/DMXPixelMapping/Source/DMXPixelMappingRenderer/Private/DMXPixelMappingPixelMapRenderer.cpp
--- a/Engine/Plugins/VirtualProduction/DMX/DMXPixelMapping/Source/DMXPixelMappingRenderer/Private/DMXPixelMappingPixelMapRenderer.cpp
+++ b/Engine/Plugins/VirtualProduction/DMX/DMXPixelMapping/Source/DMXPixelMappingRenderer/Private/DMXPixelMappingPixelMapRenderer.cpp
@@ -16,7 +16,7 @@ UDMXPixelMappingPixelMapRenderer::UDMXPixelMappingPixelMapRenderer()
 	RendererPixelMapProxy = MakeShared<FRenderPixelMapProxy>();
 }
 
-void UDMXPixelMappingPixelMapRenderer::SetElements(const TArray<TSharedRef<UE::DMXPixelMapping::Rendering::FPixelMapRenderElement>>& InElements)
+void UDMXPixelMappingPixelMapRenderer::SetElements(const TArray<TSharedRef<UE::DMXPixelMapping::Rendering::FPixelMapRenderElement>>& InElements, EPixelFormat InFormat)
 {
 	check(IsInGameThread());
 
@@ -48,15 +48,19 @@ void UDMXPixelMappingPixelMapRenderer::SetElements(const TArray<TSharedRef<UE::D
 	};
 
 	// Create a new render target if needed
+	constexpr bool bForceLinearGamma = false;
 	if (!PixelMapRenderTarget)
 	{
 		const FName TargetName = MakeUniqueObjectName(this, UTextureRenderTarget2D::StaticClass(), TEXT("PixelMapRenderTarget"));
 		PixelMapRenderTarget = NewObject<UTextureRenderTarget2D>(this, TargetName);
 		PixelMapRenderTarget->ClearColor = FLinearColor::Black;
 
-		constexpr bool bInForceLinearGamma = false;
 		const FVector2D InitialSize{ 1.0, 1.0 };
-		PixelMapRenderTarget->InitCustomFormat(ValidSize.X, ValidSize.Y, EPixelFormat::PF_B8G8R8A8, bInForceLinearGamma);
+		PixelMapRenderTarget->InitCustomFormat(ValidSize.X, ValidSize.Y, InFormat, bForceLinearGamma);
+	}
+	else if (PixelMapRenderTarget->GetFormat() != InFormat)
+	{
+		PixelMapRenderTarget->InitCustomFormat(ValidSize.X, ValidSize.Y, InFormat, bForceLinearGamma);
 	}
 	else if (ValidSize.X != PixelMapRenderTarget->GetSurfaceWidth() || ValidSize.Y != PixelMapRenderTarget->GetSurfaceHeight())
 	{
@@ -66,6 +70,12 @@ void UDMXPixelMappingPixelMapRenderer::SetElements(const TArray<TSharedRef<UE::D
 	RenderElements = InElements;
 }
 
+void UDMXPixelMappingPixelMapRenderer::SetElements(const TArray<TSharedRef<UE::DMXPixelMapping::Rendering::FPixelMapRenderElement>>& InElements)
+{
+	// DERPECATED 5.4 - Use EPixelFormat::PF_B8G8R8A8 as it was used up to 5.4
+	SetElements(InElements, EPixelFormat::PF_B8G8R8A8);
+}
+
 void UDMXPixelMappingPixelMapRenderer::Render(UTexture* InputTexture, float Brightness)
 {
 	RendererPixelMapProxy->Render(InputTexture, PixelMapRenderTarget, RenderElements, Brightness);
diff --git a/Engine/Plugins/VirtualProduction/DMX/DMXPixelMapping/Source/DMXPixelMappingRenderer/Private/DMXPixelMappingPreprocessRenderer.cpp b/Engine/Plugins/VirtualProduction/DMX/DMXPixelMapping/Source/DMXPixelMappingRenderer/Private/DMXPixelMappingPreprocessRenderer.cpp
--- a/Engine/Plugins/VirtualProduction/DMX/DMXPixelMapping/Source/DMXPixelMappingRenderer/Private/DMXPixelMappingPreprocessRenderer.cpp
+++ b/Engine/Plugins/VirtualProduction/DMX/DMXPixelMapping/Source/DMXPixelMappingRenderer/Private/DMXPixelMappingPreprocessRenderer.cpp
@@ -12,35 +12,53 @@
 #include "UObject/Package.h"
 
 
-DECLARE_CYCLE_STAT(TEXT("PixelMapping PreprocessInputTexture"), STAT_PreprocessInputTexture, STATGROUP_DMX);
+DECLARE_CYCLE_STAT(TEXT("PixelMapping PreprocessInputTexture"), STAT_DMXPixelMappingPreprocessInputTexture, STATGROUP_DMX);
 
-void UDMXPixelMappingPreprocessRenderer::SetInputTexture(UTexture* InTexture)
+void UDMXPixelMappingPreprocessRenderer::SetInputTexture(UTexture* InTexture, EPixelFormat InFormat)
 {
 	using namespace UE::DMXPixelMapping::Rendering::Preprocess::Private;
 	RenderInputProxy = MakeShared<FPreprocessRenderInputTextureProxy>(InTexture);
-	ApplyMaterialProxy = MakeShared<FPreprocessApplyFilterMaterialProxy>();
+	ApplyMaterialProxy = MakeShared<FPreprocessApplyFilterMaterialProxy>(InFormat);
 
 	bShowInputSize = false;
 }
 
-void UDMXPixelMappingPreprocessRenderer::SetInputMaterial(UMaterialInterface* InMaterial)
+void UDMXPixelMappingPreprocessRenderer::SetInputMaterial(UMaterialInterface* InMaterial, EPixelFormat InFormat)
 {
 	using namespace UE::DMXPixelMapping::Rendering::Preprocess::Private;
-	RenderInputProxy = MakeShared<FPreprocessRenderInputMaterialProxy>(InMaterial, InputSize);
-	ApplyMaterialProxy = MakeShared<FPreprocessApplyFilterMaterialProxy>();
+	RenderInputProxy = MakeShared<FPreprocessRenderInputMaterialProxy>(InMaterial, InputSize, InFormat);
+	ApplyMaterialProxy = MakeShared<FPreprocessApplyFilterMaterialProxy>(InFormat);
 
 	bShowInputSize = true;
 }
 
-void UDMXPixelMappingPreprocessRenderer::SetInputUserWidget(UUserWidget* InUserWidget)
+void UDMXPixelMappingPreprocessRenderer::SetInputUserWidget(UUserWidget* InUserWidget, EPixelFormat InFormat)
 {
 	using namespace UE::DMXPixelMapping::Rendering::Preprocess::Private;
-	RenderInputProxy = MakeShared<FPreprocessRenderInputUserWidgetProxy>(InUserWidget, InputSize);
-	ApplyMaterialProxy = MakeShared<FPreprocessApplyFilterMaterialProxy>();
+	RenderInputProxy = MakeShared<FPreprocessRenderInputUserWidgetProxy>(InUserWidget, InputSize, InFormat);
+	ApplyMaterialProxy = MakeShared<FPreprocessApplyFilterMaterialProxy>(InFormat);
 
 	bShowInputSize = true;
 }
 
+void UDMXPixelMappingPreprocessRenderer::SetInputTexture(UTexture* InTexture)
+{
+	// DERPECATED 5.4 - Use EPixelFormat::PF_B8G8R8A8 as it was used up to 5.4
+	SetInputTexture(InTexture, EPixelFormat::PF_B8G8R8A8);
+}
+
+void UDMXPixelMappingPreprocessRenderer::SetInputMaterial(UMaterialInterface* InMaterial)
+{
+	// DERPECATED 5.4 - Use EPixelFormat::PF_B8G8R8A8 as it was used up to 5.4
+	SetInputMaterial(InMaterial, EPixelFormat::PF_B8G8R8A8);
+}
+
+void UDMXPixelMappingPreprocessRenderer::SetInputUserWidget(UUserWidget* InUserWidget)
+{
+	// DERPECATED 5.4 - Use EPixelFormat::PF_B8G8R8A8 as it was used up to 5.4
+	SetInputUserWidget(InUserWidget, EPixelFormat::PF_B8G8R8A8);
+}
+
 void UDMXPixelMappingPreprocessRenderer::ClearInput()
 {
 	RenderInputProxy.Reset();
@@ -49,7 +67,7 @@ void UDMXPixelMappingPreprocessRenderer::ClearInput()
 
 void UDMXPixelMappingPreprocessRenderer::Render()
 {
-	SCOPE_CYCLE_COUNTER(STAT_PreprocessInputTexture);
+	SCOPE_CYCLE_COUNTER(STAT_DMXPixelMappingPreprocessInputTexture);
 
 	if (RenderInputProxy.IsValid())
 	{
diff --git a/Engine/Plugins/VirtualProduction/DMX/DMXPixelMapping/Source/DMXPixelMappingRenderer/Private/DMXPixelMappingRenderer.cpp b/Engine/Plugins/VirtualProduction/DMX/DMXPixelMapping/Source/DMXPixelMappingRenderer/Private/DMXPixelMappingRenderer.cpp
--- a/Engine/Plugins/VirtualProduction/DMX/DMXPixelMapping/Source/DMXPixelMappingRenderer/Private/DMXPixelMappingRenderer.cpp
+++ b/Engine/Plugins/VirtualProduction/DMX/DMXPixelMapping/Source/DMXPixelMappingRenderer/Private/DMXPixelMappingRenderer.cpp
@@ -89,9 +89,8 @@ public:
 		SHADER_PARAMETER_SAMPLER(SamplerState, InputSampler)
 
 		SHADER_PARAMETER(float, Brightness)
-		SHADER_PARAMETER(FIntPoint, InputTextureSize)
-		SHADER_PARAMETER(FIntPoint, OutputTextureSize)
-		SHADER_PARAMETER(FVector2f, UVCellSize)
+		SHADER_PARAMETER(FVector2f, UVTopLeftRotated)
+		SHADER_PARAMETER(FVector2f, UVTopRightRotated)
 	END_SHADER_PARAMETER_STRUCT()
 
 	PRAGMA_ENABLE_DEPRECATION_WARNINGS
@@ -193,9 +192,6 @@ void FDMXPixelMappingRenderer::DownsampleRender(
 					PSParameters.InputTexture = InputTextureRHI;
 					PSParameters.Brightness = Brightness;
 					PSParameters.InputSampler = TStaticSamplerState<SF_Trilinear, AM_Clamp, AM_Clamp, AM_Clamp>::GetRHI();
-					PSParameters.InputTextureSize = InputTextureSize;
-					PSParameters.OutputTextureSize = OutputTextureSize;
-					PSParameters.UVCellSize = FVector2f(PixelParam.UVCellSize);
 					SetShaderParameters(RHICmdList, PixelShader, PixelShader.GetPixelShader(), PSParameters);
 
 					// Draw a two triangle on the entire viewport.
diff --git a/Engine/Plugins/VirtualProduction/DMX/DMXPixelMapping/Source/DMXPixelMappingRenderer/Private/PreprocessApplyFilterMaterialProxy.cpp b/Engine/Plugins/VirtualProduction/DMX/DMXPixelMapping/Source/DMXPixelMappingRenderer/Private/PreprocessApplyFilterMaterialProxy.cpp
--- a/Engine/Plugins/VirtualProduction/DMX/DMXPixelMapping/Source/DMXPixelMappingRenderer/Private/PreprocessApplyFilterMaterialProxy.cpp
+++ b/Engine/Plugins/VirtualProduction/DMX/DMXPixelMapping/Source/DMXPixelMappingRenderer/Private/PreprocessApplyFilterMaterialProxy.cpp
@@ -5,18 +5,18 @@
 #include "CanvasTypes.h"
 #include "CommonRenderResources.h"
 #include "DMXStats.h"
-#include "GlobalShader.h"
-#include "RHIStaticStates.h"
-#include "ScreenRendering.h"
-#include "TextureResource.h"
 #include "Engine/Canvas.h"
 #include "Engine/Engine.h"
 #include "Engine/Texture.h"
 #include "Engine/TextureRenderTarget2D.h"
+#include "GlobalShader.h"
 #include "Materials/MaterialInstanceDynamic.h"
 #include "Materials/MaterialInterface.h"
 #include "Materials/MaterialRenderProxy.h"
 #include "Modules/ModuleManager.h"
+#include "RHIStaticStates.h"
+#include "ScreenRendering.h"
+#include "TextureResource.h"
 
 
 namespace UE::DMXPixelMapping::Rendering::Preprocess::Private
@@ -78,6 +78,10 @@ namespace UE::DMXPixelMapping::Rendering::Preprocess::Private
 		Collector.AddReferencedObject(Canvas);
 	}
 
+	FPreprocessApplyFilterMaterialProxy::FPreprocessApplyFilterMaterialProxy(EPixelFormat InFormat)
+		: Format(InFormat)
+	{}
+
 	void FPreprocessApplyFilterMaterialProxy::Render(UTexture* InInputTexture, const UDMXPixelMappingPreprocessRenderer& InPreprocessRenderer)
 	{
 		SCOPE_CYCLE_COUNTER(STAT_PreprocessApplyFilterMaterial);
@@ -159,7 +163,7 @@ namespace UE::DMXPixelMapping::Rendering::Preprocess::Private
 		}
 		UTexture* InputTexture = WeakInputTexture.Get();
 		check(InputTexture);
-
+		
 		const bool bResize = 
 			OptionalOutputSize.IsSet() &&
 			InputTexture->GetSurfaceWidth() != OptionalOutputSize.GetValue().X &&
@@ -182,7 +186,7 @@ namespace UE::DMXPixelMapping::Rendering::Preprocess::Private
 		if (bInvalidRenderTargets)
 		{
 			FlushRenderingCommands();
-
+			
 			const FVector2D InputSize = FVector2D(InputTexture->GetSurfaceWidth(), InputTexture->GetSurfaceHeight());
 
 			// Create downsample render targets
@@ -195,9 +199,12 @@ namespace UE::DMXPixelMapping::Rendering::Preprocess::Private
 				{
 					break;
 				}
+
+				constexpr bool bForceLinearGamma = false;
+
 				UTextureRenderTarget2D* DownsampleRenderTarget = NewObject<UTextureRenderTarget2D>();
 				DownsampleRenderTarget->ClearColor = FLinearColor::Black;
-				DownsampleRenderTarget->InitAutoFormat(DownsampleSize.X, DownsampleSize.Y);
+				DownsampleRenderTarget->InitCustomFormat(DownsampleSize.X, DownsampleSize.Y, Format, bForceLinearGamma);
 				DownsampleRenderTarget->UpdateResourceImmediate();
 				DownsampleRenderTargets.Add(DownsampleRenderTarget);
 			}
@@ -205,9 +212,11 @@ namespace UE::DMXPixelMapping::Rendering::Preprocess::Private
 			UTextureRenderTarget2D* ScaleRenderTarget = nullptr;
 			if (OptionalOutputSize.IsSet())
 			{
+				constexpr bool bForceLinearGamma = false;
+
 				ScaleRenderTarget = NewObject<UTextureRenderTarget2D>();
 				ScaleRenderTarget->ClearColor = FLinearColor::Black;
-				ScaleRenderTarget->InitAutoFormat(OptionalOutputSize.GetValue().X, OptionalOutputSize.GetValue().Y);
+				ScaleRenderTarget->InitCustomFormat(OptionalOutputSize.GetValue().X, OptionalOutputSize.GetValue().Y, Format, bForceLinearGamma);
 				ScaleRenderTarget->UpdateResourceImmediate();
 			}
 
diff --git a/Engine/Plugins/VirtualProduction/DMX/DMXPixelMapping/Source/DMXPixelMappingRenderer/Private/PreprocessApplyFilterMaterialProxy.h b/Engine/Plugins/VirtualProduction/DMX/DMXPixelMapping/Source/DMXPixelMappingRenderer/Private/PreprocessApplyFilterMaterialProxy.h
--- a/Engine/Plugins/VirtualProduction/DMX/DMXPixelMapping/Source/DMXPixelMappingRenderer/Private/PreprocessApplyFilterMaterialProxy.h
+++ b/Engine/Plugins/VirtualProduction/DMX/DMXPixelMapping/Source/DMXPixelMappingRenderer/Private/PreprocessApplyFilterMaterialProxy.h
@@ -47,6 +47,8 @@ namespace UE::DMXPixelMapping::Rendering::Preprocess::Private
 		, public FGCObject
 	{
 	public:
+		FPreprocessApplyFilterMaterialProxy(EPixelFormat InFormat);
+
 		//~ Begin IPreprocessApplyFilterMaterialProxy interface
 		virtual void Render(UTexture* InInputTexture, const UDMXPixelMappingPreprocessRenderer& InPreprocessRenderer) override;
 		virtual UTexture* GetRenderedTexture() const override;
@@ -68,9 +70,12 @@ namespace UE::DMXPixelMapping::Rendering::Preprocess::Private
 		/** Updates render targets */
 		void UpdateRenderTargets(int32 NumDownsamplePasses, const TOptional<FVector2D>& OptionalOutputSize);
 
-		/** Renders the Input Texture to the Output render target*/
+		/** Renders the Input Texture to the Output render target */
 		void RenderTextureToTarget(UTexture* Texture, UTextureRenderTarget2D* RenderTarget) const;
 
+		/** The pixel format his proxy should use */
+		EPixelFormat Format = PF_Unknown;
+
 		/** Weak ref to the input texture */
 		TWeakObjectPtr<UTexture> WeakInputTexture;
 
diff --git a/Engine/Plugins/VirtualProduction/DMX/DMXPixelMapping/Source/DMXPixelMappingRenderer/Private/PreprocessRenderInputMaterialProxy .cpp b/Engine/Plugins/VirtualProduction/DMX/DMXPixelMapping/Source/DMXPixelMappingRenderer/Private/PreprocessRenderInputMaterialProxy .cpp
--- a/Engine/Plugins/VirtualProduction/DMX/DMXPixelMapping/Source/DMXPixelMappingRenderer/Private/PreprocessRenderInputMaterialProxy .cpp	
+++ b/Engine/Plugins/VirtualProduction/DMX/DMXPixelMapping/Source/DMXPixelMappingRenderer/Private/PreprocessRenderInputMaterialProxy .cpp	
@@ -14,9 +14,9 @@
 
 namespace UE::DMXPixelMapping::Rendering::Preprocess::Private
 {
-	DECLARE_CYCLE_STAT(TEXT("PixelMapping RenderInputMaterial"), STAT_PreprocessRenderInputMaterial, STATGROUP_DMX);
+	DECLARE_CYCLE_STAT(TEXT("PixelMapping RenderInputMaterial"), STAT_DMXPixelMappingPreprocessRenderInputMaterial, STATGROUP_DMX);
 
-	FPreprocessRenderInputMaterialProxy::FPreprocessRenderInputMaterialProxy(UMaterialInterface* InMaterial, const FVector2D& InInputSize)
+	FPreprocessRenderInputMaterialProxy::FPreprocessRenderInputMaterialProxy(UMaterialInterface* InMaterial, const FVector2D& InInputSize, EPixelFormat InFormat)
 		: WeakMaterial(InMaterial)
 	{
 		if (InMaterial)
@@ -24,9 +24,11 @@ namespace UE::DMXPixelMapping::Rendering::Preprocess::Private
 			InMaterial->EnsureIsComplete();
 		}
 
+		constexpr bool bForceLinearGamma = false;
+
 		IntermediateRenderTarget = NewObject<UTextureRenderTarget2D>();
 		IntermediateRenderTarget->ClearColor = FLinearColor::Black;
-		IntermediateRenderTarget->InitAutoFormat(InInputSize.X, InInputSize.Y);
+		IntermediateRenderTarget->InitCustomFormat(InInputSize.X, InInputSize.Y, InFormat, bForceLinearGamma);
 		IntermediateRenderTarget->UpdateResourceImmediate();
 
 		UIMaterialBrush = MakeShared<FSlateMaterialBrush>(FVector2D(1.f));
@@ -37,7 +39,7 @@ namespace UE::DMXPixelMapping::Rendering::Preprocess::Private
 
 	void FPreprocessRenderInputMaterialProxy::Render()
 	{
-		SCOPE_CYCLE_COUNTER(STAT_PreprocessRenderInputMaterial);
+		SCOPE_CYCLE_COUNTER(STAT_DMXPixelMappingPreprocessRenderInputMaterial);
 
 		UMaterial* Material = WeakMaterial.IsValid() ? WeakMaterial->GetMaterial() : nullptr;
 		if (Material && Material->IsUIMaterial())
diff --git a/Engine/Plugins/VirtualProduction/DMX/DMXPixelMapping/Source/DMXPixelMappingRenderer/Private/PreprocessRenderInputMaterialProxy.h b/Engine/Plugins/VirtualProduction/DMX/DMXPixelMapping/Source/DMXPixelMappingRenderer/Private/PreprocessRenderInputMaterialProxy.h
--- a/Engine/Plugins/VirtualProduction/DMX/DMXPixelMapping/Source/DMXPixelMappingRenderer/Private/PreprocessRenderInputMaterialProxy.h
+++ b/Engine/Plugins/VirtualProduction/DMX/DMXPixelMapping/Source/DMXPixelMappingRenderer/Private/PreprocessRenderInputMaterialProxy.h
@@ -20,8 +20,8 @@ namespace UE::DMXPixelMapping::Rendering::Preprocess::Private
 		, public FGCObject
 	{
 	public:
-		FPreprocessRenderInputMaterialProxy(UMaterialInterface* InMaterial, const FVector2D& InInputSize);
-
+		FPreprocessRenderInputMaterialProxy(UMaterialInterface* InMaterial, const FVector2D& InInputSize, EPixelFormat Format);
+		
 		//~ Begin IPreprocessRenderInputProxy interface
 		virtual void Render() override;
 		virtual UTexture* GetRenderedTexture() const override;
diff --git a/Engine/Plugins/VirtualProduction/DMX/DMXPixelMapping/Source/DMXPixelMappingRenderer/Private/PreprocessRenderInputUserWidgetProxy.cpp b/Engine/Plugins/VirtualProduction/DMX/DMXPixelMapping/Source/DMXPixelMappingRenderer/Private/PreprocessRenderInputUserWidgetProxy.cpp
--- a/Engine/Plugins/VirtualProduction/DMX/DMXPixelMapping/Source/DMXPixelMappingRenderer/Private/PreprocessRenderInputUserWidgetProxy.cpp
+++ b/Engine/Plugins/VirtualProduction/DMX/DMXPixelMapping/Source/DMXPixelMappingRenderer/Private/PreprocessRenderInputUserWidgetProxy.cpp
@@ -5,6 +5,7 @@
 #include "Blueprint/UserWidget.h"
 #include "DMXStats.h"
 #include "Engine/TextureRenderTarget2D.h"
+#include "Framework/Application/SlateApplication.h"
 #include "Slate/WidgetRenderer.h"
 
 
@@ -12,12 +13,14 @@ namespace UE::DMXPixelMapping::Rendering::Preprocess::Private
 {
 	DECLARE_CYCLE_STAT(TEXT("PixelMapping RenderInputUserWidget"), STAT_PreprocessRenderInputUserWidget, STATGROUP_DMX);
 
-	FPreprocessRenderInputUserWidgetProxy::FPreprocessRenderInputUserWidgetProxy(UUserWidget* InUserWidget, const FVector2D& InInputSize)
+	FPreprocessRenderInputUserWidgetProxy::FPreprocessRenderInputUserWidgetProxy(UUserWidget* InUserWidget, const FVector2D& InInputSize, EPixelFormat InFormat)
 		: WeakUserWidget(InUserWidget)
 	{
+		constexpr bool bForceLinearGamma = false;
+
 		IntermediateRenderTarget = NewObject<UTextureRenderTarget2D>();
 		IntermediateRenderTarget->ClearColor = FLinearColor::Black;
-		IntermediateRenderTarget->InitAutoFormat(InInputSize.X, InInputSize.Y);
+		IntermediateRenderTarget->InitCustomFormat(InInputSize.X, InInputSize.Y, InFormat, bForceLinearGamma);
 		IntermediateRenderTarget->UpdateResourceImmediate();
 
 		const bool bUseGammaCorrection = false;
diff --git a/Engine/Plugins/VirtualProduction/DMX/DMXPixelMapping/Source/DMXPixelMappingRenderer/Private/PreprocessRenderInputUserWidgetProxy.h b/Engine/Plugins/VirtualProduction/DMX/DMXPixelMapping/Source/DMXPixelMappingRenderer/Private/PreprocessRenderInputUserWidgetProxy.h
--- a/Engine/Plugins/VirtualProduction/DMX/DMXPixelMapping/Source/DMXPixelMappingRenderer/Private/PreprocessRenderInputUserWidgetProxy.h
+++ b/Engine/Plugins/VirtualProduction/DMX/DMXPixelMapping/Source/DMXPixelMappingRenderer/Private/PreprocessRenderInputUserWidgetProxy.h
@@ -19,7 +19,7 @@ namespace UE::DMXPixelMapping::Rendering::Preprocess::Private
 		, public FGCObject
 	{
 	public:
-		FPreprocessRenderInputUserWidgetProxy(UUserWidget* InUserWidget, const FVector2D& InInputSize);
+		FPreprocessRenderInputUserWidgetProxy(UUserWidget* InUserWidget, const FVector2D& InInputSize, EPixelFormat InFormat);
 
 		//~ Begin IPreprocessRenderInputProxy interface
 		virtual void Render() override;
diff --git a/Engine/Plugins/VirtualProduction/DMX/DMXPixelMapping/Source/DMXPixelMappingRenderer/Private/RenderPixelMapProxy.cpp b/Engine/Plugins/VirtualProduction/DMX/DMXPixelMapping/Source/DMXPixelMappingRenderer/Private/RenderPixelMapProxy.cpp
--- a/Engine/Plugins/VirtualProduction/DMX/DMXPixelMapping/Source/DMXPixelMappingRenderer/Private/RenderPixelMapProxy.cpp
+++ b/Engine/Plugins/VirtualProduction/DMX/DMXPixelMapping/Source/DMXPixelMappingRenderer/Private/RenderPixelMapProxy.cpp
@@ -58,9 +58,8 @@ public:
 		SHADER_PARAMETER_SAMPLER(SamplerState, InputSampler)
 
 		SHADER_PARAMETER(float, Brightness)
-		SHADER_PARAMETER(FIntPoint, InputTextureSize)
-		SHADER_PARAMETER(FIntPoint, OutputTextureSize)
-		SHADER_PARAMETER(FVector2f, UVCellSize)
+		SHADER_PARAMETER(FVector2f, UVTopLeftRotated)
+		SHADER_PARAMETER(FVector2f, UVTopRightRotated)
 	END_SHADER_PARAMETER_STRUCT()
 
 	static bool ShouldCompilePermutation(const FGlobalShaderPermutationParameters& Parameters) { return true; }
@@ -123,7 +122,8 @@ namespace UE::DMXPixelMapping::Rendering::Private
 
 						const FVector2D& UV = Element->GetParameters().UV;
 						const FVector2D& UVSize = Element->GetParameters().UVSize;
-						const FVector2D& UVCellSize = Element->GetParameters().UVCellSize;
+						const FVector2D& UVTopLeftRotated = Element->GetParameters().UVTopLeftRotated;
+						const FVector2D& UVTopRightRotated = Element->GetParameters().UVTopRightRotated;
 
 						const int32 RenderTargetPoisitionX = PixelIndex % OutputTextureSize.X;
 						const int32 RenderTargetPositionY = PixelIndex / OutputTextureSize.X;
@@ -164,9 +164,8 @@ namespace UE::DMXPixelMapping::Rendering::Private
 						PSParameters.InputTexture = InputTextureRHI;
 						PSParameters.Brightness = Brightness;
 						PSParameters.InputSampler = TStaticSamplerState<SF_Bilinear, AM_Clamp, AM_Clamp, AM_Clamp>::GetRHI();
-						PSParameters.InputTextureSize = InputTextureSize;
-						PSParameters.OutputTextureSize = OutputTextureSize;
-						PSParameters.UVCellSize = FVector2f(UVCellSize);
+						PSParameters.UVTopLeftRotated = FVector2f(UVTopLeftRotated);
+						PSParameters.UVTopRightRotated = FVector2f(UVTopRightRotated);
 						SetShaderParameters(RHICmdList, PixelShader, PixelShader.GetPixelShader(), PSParameters);
 
 						// Draw a two triangle on the entire viewport.
@@ -181,10 +180,11 @@ namespace UE::DMXPixelMapping::Rendering::Private
 					TArray<FLinearColor> ColorArray;
 					const FIntRect Rect(0, 0, OutputTextureSize.X, OutputTextureSize.Y);
 
-					// Read surface without flush rendering thread
-					RHICmdList.ReadSurfaceData(RenderTargetResource->TextureRHI, Rect, ColorArray, FReadSurfaceDataFlags());
+					// Read surface data. 
+					// Use min/max compression mode here - Either the input was tone mapped earlier, or we want to preserve dynamic range.
+					RHICmdList.ReadSurfaceData(RenderTargetResource->TextureRHI, Rect, ColorArray, ERangeCompressionMode::RCM_MinMax);
 
-					// Shrink color array to match num elements
+					// Write data to the game thread
 					if (!ensureMsgf(ColorArray.Num() >= Elements.Num(), TEXT("Pixel maparray is smaller than the number of elements that should have been rendered. Failed to Pixel Map")))
 					{
 						return;
diff --git a/Engine/Plugins/VirtualProduction/DMX/DMXPixelMapping/Source/DMXPixelMappingRenderer/Public/DMXPixelMappingPixelMapRenderer.h b/Engine/Plugins/VirtualProduction/DMX/DMXPixelMapping/Source/DMXPixelMappingRenderer/Public/DMXPixelMappingPixelMapRenderer.h
--- a/Engine/Plugins/VirtualProduction/DMX/DMXPixelMapping/Source/DMXPixelMappingRenderer/Public/DMXPixelMappingPixelMapRenderer.h
+++ b/Engine/Plugins/VirtualProduction/DMX/DMXPixelMapping/Source/DMXPixelMappingRenderer/Public/DMXPixelMappingPixelMapRenderer.h
@@ -2,6 +2,7 @@
 
 #pragma once
 
+#include "PixelFormat.h"
 #include "UObject/Object.h"
 
 #include "DMXPixelMappingPixelMapRenderer.generated.h"
@@ -54,6 +55,10 @@ public:
 	UDMXPixelMappingPixelMapRenderer();
 
 	/** Sets the elements to render */
+	void SetElements(const TArray<TSharedRef<UE::DMXPixelMapping::Rendering::FPixelMapRenderElement>>& InElements, EPixelFormat InFormat);
+
+	/** DEPRECATED 5.4 - Sets the elements to render */
+	UE_DEPRECATED(5.4, "It is now required to specify a pixel format, see related overload.")
 	void SetElements(const TArray<TSharedRef<UE::DMXPixelMapping::Rendering::FPixelMapRenderElement>>& InElements);
 
 	/** Renders the current elements */
diff --git a/Engine/Plugins/VirtualProduction/DMX/DMXPixelMapping/Source/DMXPixelMappingRenderer/Public/DMXPixelMappingPreprocessRenderer.h b/Engine/Plugins/VirtualProduction/DMX/DMXPixelMapping/Source/DMXPixelMappingRenderer/Public/DMXPixelMappingPreprocessRenderer.h
--- a/Engine/Plugins/VirtualProduction/DMX/DMXPixelMapping/Source/DMXPixelMappingRenderer/Public/DMXPixelMappingPreprocessRenderer.h
+++ b/Engine/Plugins/VirtualProduction/DMX/DMXPixelMapping/Source/DMXPixelMappingRenderer/Public/DMXPixelMappingPreprocessRenderer.h
@@ -2,6 +2,7 @@
 
 #pragma once
 
+#include "PixelFormat.h"
 #include "UObject/Object.h"
 
 #include "DMXPixelMappingPreprocessRenderer.generated.h"
@@ -64,12 +65,24 @@ class DMXPIXELMAPPINGRENDERER_API UDMXPixelMappingPreprocessRenderer
 
 public:
 	/** Sets the input texture that is used when rendering */
-	void SetInputTexture(UTexture* InTexture);
+	void SetInputTexture(UTexture* InTexture, EPixelFormat InFormat);
 
 	/** Sets the input texture that is used when rendering */
-	void SetInputMaterial(UMaterialInterface* InMaterial);
+	void SetInputMaterial(UMaterialInterface* InMaterial, EPixelFormat InFormat);
 
 	/** Sets the input texture that is used when rendering */
+	void SetInputUserWidget(UUserWidget* InUserWidget, EPixelFormat InFormat);
+
+	/** DEPRECATED 5.4 - Sets the input texture that is used when rendering */
+	UE_DEPRECATED(5.4, "It is now required to specify a pixel format, see related overload.")
+	void SetInputTexture(UTexture* InTexture);
+
+	/** DEPRECATED 5.4 - Sets the input texture that is used when rendering */
+	UE_DEPRECATED(5.4, "It is now required to specify a pixel format, see related overload.")
+	void SetInputMaterial(UMaterialInterface* InMaterial);
+
+	/** DEPRECATED 5.4 - Sets the input texture that is used when rendering */
+	UE_DEPRECATED(5.4, "It is now required to specify a pixel format, see related overload.")
 	void SetInputUserWidget(UUserWidget* InUserWidget);
 
 	/** Clears input */
diff --git a/Engine/Plugins/VirtualProduction/DMX/DMXPixelMapping/Source/DMXPixelMappingRenderer/Public/DMXPixelMappingRenderElement.h b/Engine/Plugins/VirtualProduction/DMX/DMXPixelMapping/Source/DMXPixelMappingRenderer/Public/DMXPixelMappingRenderElement.h
--- a/Engine/Plugins/VirtualProduction/DMX/DMXPixelMapping/Source/DMXPixelMappingRenderer/Public/DMXPixelMappingRenderElement.h
+++ b/Engine/Plugins/VirtualProduction/DMX/DMXPixelMapping/Source/DMXPixelMappingRenderer/Public/DMXPixelMappingRenderElement.h
@@ -34,14 +34,23 @@ namespace UE::DMXPixelMapping::Rendering
 {
 	struct DMXPIXELMAPPINGRENDERER_API FPixelMapRenderElementParameters
 	{
-		/** Position in texels of the top left corner of the quad's UV's */
+		/** Position in texels of the center of the quad's UV's */
 		FVector2D UV;
 
-		/** Size in texels of the quad's total UV space */
+		/** Size in texels of UV. May match UVSize */
 		FVector2D UVSize;
 
-		/** Size in texels of UV.May match UVSize */
-		FVector2D UVCellSize;
+		/** Position in texels of the top left corner of the quad's UV's, in rotated space */
+		FVector2D UVTopLeftRotated;
+
+		/** Position in texels of the top right corner of the quad's UV's, in rotate space */
+		FVector2D UVTopRightRotated;
+
+		/** Rotation in degrees. Useful to paint these elements in the preview view. */
+		double Rotation;
+
+		/** DERPECATED 5.4 - Size in texels of the quad's total UV space */
+		FVector2D UVCellSize_DEPRECATED;
 
 		/** The quality of color samples in the pixel shader(number of samples) */
 		EDMXPixelBlendingQuality CellBlendingQuality;
diff --git a/Engine/Plugins/VirtualProduction/DMX/DMXPixelMapping/Source/DMXPixelMappingRuntime/Private/Components/DMXPixelMappingFixtureGroupItemComponent.cpp b/Engine/Plugins/VirtualProduction/DMX/DMXPixelMapping/Source/DMXPixelMappingRuntime/Private/Components/DMXPixelMappingFixtureGroupItemComponent.cpp
--- a/Engine/Plugins/VirtualProduction/DMX/DMXPixelMapping/Source/DMXPixelMappingRuntime/Private/Components/DMXPixelMappingFixtureGroupItemComponent.cpp
+++ b/Engine/Plugins/VirtualProduction/DMX/DMXPixelMapping/Source/DMXPixelMappingRuntime/Private/Components/DMXPixelMappingFixtureGroupItemComponent.cpp
@@ -427,7 +427,16 @@ void UDMXPixelMappingFixtureGroupItemComponent::UpdateRenderElement()
 	FPixelMapRenderElementParameters Parameters;
 	Parameters.UV = FVector2D(GetPosition().X / InputTextureWidth, GetPosition().Y / InputTextureHeight);
 	Parameters.UVSize = FVector2D(GetSize().X / InputTextureWidth, GetSize().Y / InputTextureHeight);
-	Parameters.UVCellSize = Parameters.UVSize / 2.f;
+
+	FVector2D A;
+	FVector2D B;
+	FVector2D C;
+	FVector2D D;
+	GetEdges(A, B, C, D);
+	Parameters.UVTopLeftRotated = FVector2D(A.X / InputTextureWidth, A.Y / InputTextureHeight);
+	Parameters.UVTopRightRotated = FVector2D(B.X / InputTextureWidth, B.Y / InputTextureHeight);
+
+	Parameters.Rotation = GetRotation();
 	Parameters.CellBlendingQuality = CellBlendingQuality;
 	Parameters.bStaticCalculateUV = true;
 
diff --git a/Engine/Plugins/VirtualProduction/DMX/DMXPixelMapping/Source/DMXPixelMappingRuntime/Private/Components/DMXPixelMappingMatrixCellComponent.cpp b/Engine/Plugins/VirtualProduction/DMX/DMXPixelMapping/Source/DMXPixelMappingRuntime/Private/Components/DMXPixelMappingMatrixCellComponent.cpp
--- a/Engine/Plugins/VirtualProduction/DMX/DMXPixelMapping/Source/DMXPixelMappingRuntime/Private/Components/DMXPixelMappingMatrixCellComponent.cpp
+++ b/Engine/Plugins/VirtualProduction/DMX/DMXPixelMapping/Source/DMXPixelMappingRuntime/Private/Components/DMXPixelMappingMatrixCellComponent.cpp
@@ -215,7 +215,17 @@ void UDMXPixelMappingMatrixCellComponent::UpdateRenderElement()
 	FPixelMapRenderElementParameters Parameters;
 	Parameters.UV = FVector2D(GetPosition().X / InputTextureWidth, GetPosition().Y / InputTextureHeight);
 	Parameters.UVSize = FVector2D(GetSize().X / InputTextureWidth, GetSize().Y / InputTextureHeight);
-	Parameters.UVCellSize = Parameters.UVSize / 2.f;
+
+	FVector2D A;
+	FVector2D B;
+	FVector2D C;
+	FVector2D D;
+	GetEdges(A, B, C, D);
+
+	Parameters.UVTopLeftRotated = FVector2D(A.X / InputTextureWidth, A.Y / InputTextureHeight);
+	Parameters.UVTopRightRotated = FVector2D(B.X / InputTextureWidth, B.Y / InputTextureHeight);
+
+	Parameters.Rotation = GetRotation();
 	Parameters.CellBlendingQuality = MatrixComponent->CellBlendingQuality;
 	Parameters.bStaticCalculateUV = true;
 
diff --git a/Engine/Plugins/VirtualProduction/DMX/DMXPixelMapping/Source/DMXPixelMappingRuntime/Private/Components/DMXPixelMappingRendererComponent.cpp b/Engine/Plugins/VirtualProduction/DMX/DMXPixelMapping/Source/DMXPixelMappingRuntime/Private/Components/DMXPixelMappingRendererComponent.cpp
--- a/Engine/Plugins/VirtualProduction/DMX/DMXPixelMapping/Source/DMXPixelMappingRuntime/Private/Components/DMXPixelMappingRendererComponent.cpp
+++ b/Engine/Plugins/VirtualProduction/DMX/DMXPixelMapping/Source/DMXPixelMappingRuntime/Private/Components/DMXPixelMappingRendererComponent.cpp
@@ -15,6 +15,7 @@
 #include "DMXPixelMappingMainStreamObjectVersion.h"
 #include "DMXPixelMappingTypes.h"
 #include "DMXStats.h"
+#include "Engine/Texture2D.h"
 #include "Engine/TextureRenderTarget2D.h"
 #include "IDMXPixelMappingRenderer.h"
 #include "IDMXPixelMappingRendererModule.h"
@@ -120,10 +121,15 @@ void UDMXPixelMappingRendererComponent::PostEditChangeChainProperty(FPropertyCha
 	{
 		UpdatePreprocessRenderer();
 		LetChildrenFollowSize();
+	}	
+	
+	const FName PropertyName = PropertyChangedChainEvent.GetPropertyName();
+	if (PropertyName == GET_MEMBER_NAME_CHECKED(UDMXPixelMappingRendererComponent, DynamicRange))
+	{
+		InvalidatePixelMapRenderer();
 	}
 
 	PRAGMA_DISABLE_DEPRECATION_WARNINGS
-	const FName PropertyName = PropertyChangedChainEvent.GetPropertyName();
 	if (PropertyName == GET_MEMBER_NAME_CHECKED(UDMXPixelMappingRendererComponent, Brightness))
 	{
 		const TSharedPtr<IDMXPixelMappingRenderer>& Renderer = GetRenderer();
@@ -148,19 +154,21 @@ void UDMXPixelMappingRendererComponent::UpdatePreprocessRenderer()
 	PRAGMA_ENABLE_DEPRECATION_WARNINGS
 
 	UserWidget = nullptr;
+
+	const EPixelFormat Format = GetFormatFromDynamicRange();
 	switch (RendererType)
 	{
 	case(EDMXPixelMappingRendererType::Texture):
-		PreprocessRenderer->SetInputTexture(InputTexture.Get());
+		PreprocessRenderer->SetInputTexture(InputTexture.Get(), Format);
 		break;
 
 	case(EDMXPixelMappingRendererType::Material):
-		PreprocessRenderer->SetInputMaterial(InputMaterial.Get());
+		PreprocessRenderer->SetInputMaterial(InputMaterial.Get(), Format);
 		break;
 
 	case(EDMXPixelMappingRendererType::UMG):
 		UserWidget = CreateWidget(TryGetWorld(), InputWidget);
-		PreprocessRenderer->SetInputUserWidget(UserWidget.Get());
+		PreprocessRenderer->SetInputUserWidget(UserWidget.Get(), Format);
 		break;
 
 	default:
@@ -266,7 +274,7 @@ void UDMXPixelMappingRendererComponent::Render()
 		constexpr bool bRecursive = true;
 		ForEachChild([this](UDMXPixelMappingBaseComponent* Component)
 			{
-				if (UDMXPixelMappingFixtureGroupItemComponent* FixtureGroupItemComponent = Cast< UDMXPixelMappingFixtureGroupItemComponent>(Component))
+				if (UDMXPixelMappingFixtureGroupItemComponent* FixtureGroupItemComponent = Cast<UDMXPixelMappingFixtureGroupItemComponent>(Component))
 				{
 					PixelMapRenderElements.Add(FixtureGroupItemComponent->GetOrCreatePixelMapRenderElement());
 				}
@@ -276,7 +284,8 @@ void UDMXPixelMappingRendererComponent::Render()
 				}
 			}, bRecursive);
 
-		PixelMapRenderer->SetElements(PixelMapRenderElements);
+		const EPixelFormat Format = GetFormatFromDynamicRange();
+		PixelMapRenderer->SetElements(PixelMapRenderElements, Format);
 
 		bInvalidatePixelMap = false;
 	}
@@ -345,7 +354,7 @@ void UDMXPixelMappingRendererComponent::LetChildrenFollowSize()
 		LayoutRect = NewSize;
 		return;
 	}
-
+	
 	// Skip unchanged values, or if the current size is zero (no texture).
 	if (NewSize == LayoutRect ||
 		NewSize == FVector2D::ZeroVector)
@@ -375,6 +384,45 @@ void UDMXPixelMappingRendererComponent::LetChildrenFollowSize()
 	LayoutRect = NewSize;
 }
 
+EPixelFormat UDMXPixelMappingRendererComponent::GetFormatFromDynamicRange() const
+{
+	UTexture2D* InputTexture2D = Cast<UTexture2D>(InputTexture);
+
+	if (DynamicRange == EDMXPixelMappingRendererDynamicRange::Auto &&
+		RendererType == EDMXPixelMappingRendererType::Texture &&
+		InputTexture2D)
+	{
+		const EPixelFormat PixelFormat = InputTexture2D->GetPixelFormat();
+	
+		// Propagonate the pixel format, if it is supported by render targets. 
+		// As there doesn't seem to be an engine call to do this conversion,
+		// use (the opposite) logic of GetPixelFormatFromRenderTargetFormat in TextureRenderTarget2D.h (5.4).
+		if (PixelFormat == PF_G8 ||
+			PixelFormat == PF_R8G8 ||
+			PixelFormat == PF_B8G8R8A8 ||
+			PixelFormat == PF_R16F ||
+			PixelFormat == PF_G16R16F ||
+			PixelFormat == PF_FloatRGBA ||
+			PixelFormat == PF_R32_FLOAT ||
+			PixelFormat == PF_G32R32F ||
+			PixelFormat == PF_A32B32G32R32F ||
+			PixelFormat == PF_A2B10G10R10)
+		{
+			return PixelFormat;
+		}
+		else
+		{
+			return PF_FloatRGBA;
+		}
+	}
+	else if (DynamicRange == EDMXPixelMappingRendererDynamicRange::RGBA16F)
+	{
+		return PF_FloatRGBA;
+	}
+
+	return PF_B8G8R8A8;
+}
+
 UWorld* UDMXPixelMappingRendererComponent::TryGetWorld() const
 {
 	UWorld* World = nullptr;
diff --git a/Engine/Plugins/VirtualProduction/DMX/DMXPixelMapping/Source/DMXPixelMappingRuntime/Public/Components/DMXPixelMappingRendererComponent.h b/Engine/Plugins/VirtualProduction/DMX/DMXPixelMapping/Source/DMXPixelMappingRuntime/Public/Components/DMXPixelMappingRendererComponent.h
--- a/Engine/Plugins/VirtualProduction/DMX/DMXPixelMapping/Source/DMXPixelMappingRuntime/Public/Components/DMXPixelMappingRendererComponent.h
+++ b/Engine/Plugins/VirtualProduction/DMX/DMXPixelMapping/Source/DMXPixelMappingRuntime/Public/Components/DMXPixelMappingRendererComponent.h
@@ -22,6 +22,15 @@ class UWorld;
 namespace UE::DMXPixelMapping::Rendering { class FPixelMapRenderElement; }
 
 
+UENUM(BlueprintType)
+enum class EDMXPixelMappingRendererDynamicRange : uint8
+{
+	Auto UMETA(ToolTip = "Use the format of the Input Texture, use 8-bit for Input Materials and UMG."),
+	RGBA8 UMETA(DisplayName = "Low Precision 8-bit (RGBA8)"),
+	RGBA16F UMETA(DisplayName = "High Precision 16-bit (RGBA16F)")
+};
+
+
 /** 
  * Component for rendering input texture.  
  */
@@ -98,8 +107,12 @@ public:
 	UPROPERTY(EditAnywhere, BlueprintReadOnly, Category = "Render Settings")
 	TSubclassOf<UUserWidget> InputWidget;
 
-	/** The brightness of the renderer */
-	UPROPERTY(EditAnywhere, BlueprintReadWrite, Category = "Render Settings", Meta = (ClampMin = "0", UIMin = "0", UIMax = "1"))
+	/** The dynamic range of the renderer */
+	UPROPERTY(EditAnywhere, BlueprintReadOnly, Category = "Render Settings")
+	EDMXPixelMappingRendererDynamicRange DynamicRange = EDMXPixelMappingRendererDynamicRange::Auto;
+
+	/** The exposure of the pixel mapping renderer. This property affects the DMX output. */
+	UPROPERTY(EditAnywhere, BlueprintReadWrite, Category = "Render Settings", Meta = (ClampMin = "0", UIMin = "0", UIMax = "1", DisplayName = "Exposure"))
 	float Brightness = 1.f;
 
 	/** If true, children are positioned relative to the size of this renderer */
@@ -117,6 +130,9 @@ private:
 	/** When bChildrenFollowSize is true, rearranges children relatively to the current size */
 	void LetChildrenFollowSize();
 
+	/** Returns the pixel format depending on the current dynamic range setting */
+	EPixelFormat GetFormatFromDynamicRange() const;
+
 	/** Tries to get any world. If with editor, returns the editor world, in game returns GWorld */
 	UWorld* TryGetWorld() const;
 
@@ -136,7 +152,7 @@ private:
 
 	/** Renderer responsible to pixel map */
 	UPROPERTY()
-	UDMXPixelMappingPixelMapRenderer* PixelMapRenderer;
+	TObjectPtr<UDMXPixelMappingPixelMapRenderer> PixelMapRenderer;
 
 	/** The current rect in which children are laid out. Useful to compute UV position and size of children. */
 	UPROPERTY()
diff --git a/Engine/Plugins/VirtualProduction/DMX/DMXPixelMapping/Source/DMXPixelMappingRuntime/Public/DMXPixelMapping.h b/Engine/Plugins/VirtualProduction/DMX/DMXPixelMapping/Source/DMXPixelMappingRuntime/Public/DMXPixelMapping.h
--- a/Engine/Plugins/VirtualProduction/DMX/DMXPixelMapping/Source/DMXPixelMappingRuntime/Public/DMXPixelMapping.h
+++ b/Engine/Plugins/VirtualProduction/DMX/DMXPixelMapping/Source/DMXPixelMappingRuntime/Public/DMXPixelMapping.h
@@ -144,6 +144,9 @@ public:
 	TObjectPtr<UDMXPixelMappingRootComponent> RootComponent;
 
 #if WITH_EDITORONLY_DATA
+	////////////////////////////////////
+	// Per Asset Editor User Settings
+
 	/** If true, grid snapping is enabled */
 	UPROPERTY(NonTransactional)
 	bool bGridSnappingEnabled = false;
@@ -164,6 +167,10 @@ public:
 	UPROPERTY(NonTransactional)
 	float ComponentLabelFontSize = 8.f;
 
+	/** Exposure of the designer view */
+	UPROPERTY(NonTransactional)
+	float DesignerExposure = 1.f;
+
 	/** If true, new components use the fixture patch color instead of the default pixel mapping color. */
 	UPROPERTY(NonTransactional)
 	bool bNewComponentsUsePatchColor = true;
