DMX - Rework Pixel Mapping play and stop behaviour, add pause behaviour

diff --git a/Engine/Plugins/VirtualProduction/DMX/DMXEngine/Source/DMXRuntime/Private/Library/DMXEntityFixturePatch.cpp b/Engine/Plugins/VirtualProduction/DMX/DMXEngine/Source/DMXRuntime/Private/Library/DMXEntityFixturePatch.cpp
--- a/Engine/Plugins/VirtualProduction/DMX/DMXEngine/Source/DMXRuntime/Private/Library/DMXEntityFixturePatch.cpp
+++ b/Engine/Plugins/VirtualProduction/DMX/DMXEngine/Source/DMXRuntime/Private/Library/DMXEntityFixturePatch.cpp
@@ -299,8 +299,8 @@ void UDMXEntityFixturePatch::SendDMX(TMap<FDMXAttributeName, int32> AttributeMap
 	{
 		for (const FDMXFixtureFunction& Function : ModePtr->Functions)
 		{
-			const FDMXAttributeName& FunctionAttr = Function.Attribute;
-			if (FunctionAttr == Elem.Key)
+			const FDMXAttributeName& FunctionAttribute = Function.Attribute;
+			if (FunctionAttribute == Elem.Key)
 			{
 				const int32 LastFunctionChannel = Function.Channel + Function.GetNumChannels() - 1;
 				const bool bLastChannelExceedsChannelSpan = LastFunctionChannel > ModePtr->ChannelSpan;
@@ -334,6 +334,18 @@ void UDMXEntityFixturePatch::SendDMX(TMap<FDMXAttributeName, int32> AttributeMap
 	}
 }
 
+void UDMXEntityFixturePatch::SendDefaultValues()
+{
+	constexpr bool bSendDefaultValues = true;
+	SendResetDataToAllAttributes(bSendDefaultValues);
+}
+
+void UDMXEntityFixturePatch::SendZeroValues()
+{
+	constexpr bool bSendDefaultValues = false;
+	SendResetDataToAllAttributes(bSendDefaultValues);
+}
+
 void UDMXEntityFixturePatch::RebuildCache()
 {
 	LastDMXSignal.Reset();
@@ -1207,6 +1219,76 @@ bool UDMXEntityFixturePatch::GetAllMatrixCells(TArray<FDMXCell>& Cells)
 	return false;
 }
 
+void UDMXEntityFixturePatch::SendResetDataToAllAttributes(bool bUseDefaultValues)
+{
+	UDMXLibrary* DMXLibrary = GetParentLibrary();
+	const FDMXFixtureMode* ModePtr = GetActiveMode();
+	if (!DMXLibrary || !ModePtr)
+	{
+		return;
+	}
+
+	// Non-matrix channels
+	TMap<int32, uint8> DMXChannelToValueMap;
+	for (const FDMXFixtureFunction& Function : ModePtr->Functions)
+	{
+		const FDMXAttributeName& FunctionAttribute = Function.Attribute;
+
+		const int32 LastFunctionChannel = Function.Channel + Function.GetNumChannels() - 1;
+		const bool bLastChannelExceedsChannelSpan = LastFunctionChannel > ModePtr->ChannelSpan;
+		const bool bLastChannelExceedsUniverseSize = LastFunctionChannel + GetStartingChannel() - 1 > DMX_MAX_ADDRESS;
+
+		if (!bLastChannelExceedsChannelSpan && !bLastChannelExceedsUniverseSize)
+		{
+			const int32 Channel = Function.Channel + GetStartingChannel() - 1;
+
+			uint32 Value = bUseDefaultValues ? Function.DefaultValue : 0;
+			uint8* ChannelValueBytes = reinterpret_cast<uint8*>(&Value);
+			UDMXEntityFixtureType::FunctionValueToBytes(Function, Value, ChannelValueBytes);
+
+			const uint8 NumBytesInSignalFormat = FDMXConversions::GetSizeOfSignalFormat(Function.DataType);
+			for (uint8 ChannelIt = 0; ChannelIt < NumBytesInSignalFormat; ++ChannelIt)
+			{
+				DMXChannelToValueMap.Add(Channel + ChannelIt, ChannelValueBytes[ChannelIt]);
+			}
+		}
+	}
+
+	// Matrix channels
+	const FDMXFixtureMatrix* const FixtureMatrixPtr = GetFixtureMatrix();
+	if (FixtureMatrixPtr)
+	{
+		TArray<FDMXCell> MatrixCells;
+		GetAllMatrixCells(MatrixCells);
+
+		for (const FDMXCell& MatrixCell : MatrixCells)
+		{
+			const int32 DistributedCellIndex = Cache.GetDistributedCellIndex(MatrixCell.Coordinate);
+			const int32 AbsoluteMatrixStartingChannel = Cache.GetMatrixStartingChannelAbsolute();
+			const int32 AbsoluteCellStartingChannel = AbsoluteMatrixStartingChannel + DistributedCellIndex * Cache.GetCellSize();
+
+			int32 AttributeOffset = 0;
+			for (const FDMXFixtureCellAttribute& CellAttribute : Cache.GetCellAttributes())
+			{
+				const uint32 Value = bUseDefaultValues ? CellAttribute.DefaultValue : 0;
+				const TArray<uint8> ByteArray = FDMXConversions::UnsignedInt32ToByteArray(Value, CellAttribute.DataType, CellAttribute.bUseLSBMode);
+
+				DMXChannelToValueMap.Reserve(ByteArray.Num());
+				for (int32 ByteIndex = 0; ByteIndex < CellAttribute.GetNumChannels(); ByteIndex++)
+				{
+					DMXChannelToValueMap.Add(AbsoluteCellStartingChannel + AttributeOffset + ByteIndex, ByteArray[ByteIndex]);
+				}
+			}
+		}
+	}
+
+	// Send to the library's output ports
+	for (const FDMXOutputPortSharedRef& OutputPort : DMXLibrary->GetOutputPorts())
+	{
+		OutputPort->SendDMX(UniverseID, DMXChannelToValueMap);
+	}
+}
+
 void UDMXEntityFixturePatch::OnFixtureTypeChanged(const UDMXEntityFixtureType* FixtureType)
 {	
 	if (ParentFixtureTypeTemplate == FixtureType)
diff --git a/Engine/Plugins/VirtualProduction/DMX/DMXEngine/Source/DMXRuntime/Public/Library/DMXEntityFixturePatch.h b/Engine/Plugins/VirtualProduction/DMX/DMXEngine/Source/DMXRuntime/Public/Library/DMXEntityFixturePatch.h
--- a/Engine/Plugins/VirtualProduction/DMX/DMXEngine/Source/DMXRuntime/Public/Library/DMXEntityFixturePatch.h
+++ b/Engine/Plugins/VirtualProduction/DMX/DMXEngine/Source/DMXRuntime/Public/Library/DMXEntityFixturePatch.h
@@ -114,6 +114,12 @@ public:
 	UFUNCTION(BlueprintCallable, Category = "DMX")
 	void SendDMX(TMap<FDMXAttributeName, int32> AttributeMap);
 
+	/** Sends the default value for all attributes, including matrix attributes */
+	void SendDefaultValues();
+
+	/** Sends zeroes for all attributes, including matrix attributes */
+	void SendZeroValues();
+
 	/** 
 	 * Rebuilds the cache. Should be called when relevant properties, for example the starting channel changed. 
 	 * This will not clear cached DMX data.
@@ -476,6 +482,9 @@ public:
 	bool GetAllMatrixCells(TArray<FDMXCell>& Cells);
 
 private:
+	/** Sends reset data to all channels. If bUseDefaultValues is true, sends default values. If it's false it sends zeroes. */
+	void SendResetDataToAllAttributes(bool bUseDefaultValues);
+
 	/** Called when a Fixture Type changed */
 	void OnFixtureTypeChanged(const UDMXEntityFixtureType* FixtureType);
 
diff --git a/Engine/Plugins/VirtualProduction/DMX/DMXPixelMapping/Source/DMXPixelMappingEditor/Private/DMXPixelMappingEditorCommands.cpp b/Engine/Plugins/VirtualProduction/DMX/DMXPixelMapping/Source/DMXPixelMappingEditor/Private/DMXPixelMappingEditorCommands.cpp
--- a/Engine/Plugins/VirtualProduction/DMX/DMXPixelMapping/Source/DMXPixelMappingEditor/Private/DMXPixelMappingEditorCommands.cpp
+++ b/Engine/Plugins/VirtualProduction/DMX/DMXPixelMapping/Source/DMXPixelMappingEditor/Private/DMXPixelMappingEditorCommands.cpp
@@ -19,8 +19,16 @@ FDMXPixelMappingEditorCommands::FDMXPixelMappingEditorCommands()
 void FDMXPixelMappingEditorCommands::RegisterCommands()
 {
 	UI_COMMAND(AddMapping, "Add Source", "Adds a new Source Texure, Material or User Widget to the Pixel Map asset.", EUserInterfaceActionType::Button, FInputChord());
-	UI_COMMAND(PlayDMX, "Play DMX", "Starts to send DMX", EUserInterfaceActionType::Button, FInputChord());
-	UI_COMMAND(StopPlayingDMX, "Stop Playing DMX", "Stops playing DMX.", EUserInterfaceActionType::Button, FInputChord());
+	UI_COMMAND(PlayDMX, "Plays DMX", "Plays DMX", EUserInterfaceActionType::Button, FInputChord());
+	UI_COMMAND(PauseDMX, "Pause DMX", "Pauses playing DMX. DMX values will still be sent, at a lower rate.", EUserInterfaceActionType::Button, FInputChord());
+	UI_COMMAND(ResumeDMX, "Resume DMX", "Resumes playing DMX after being paused. ", EUserInterfaceActionType::Button, FInputChord());
+	UI_COMMAND(StopDMX, "Stop DMX", "Stops playing DMX.", EUserInterfaceActionType::Button, FInputChord());
+	UI_COMMAND(TogglePlayPauseDMX, "Toggle Play/Pause DMX", "Toggles between playing and pausing DMX", EUserInterfaceActionType::Button, FInputChord(EModifierKey::Shift, EKeys::SpaceBar));
+	UI_COMMAND(TogglePlayStopDMX, "Toggle Play/Stop DMX", "Toggles between playing and stopping DMX", EUserInterfaceActionType::Button, FInputChord(EKeys::SpaceBar));
+
+	UI_COMMAND(EditorStopSendsZeroValues, "Stop sends Zero Values", "When stop is clicked in the editor, zero values are sent to all patches in use", EUserInterfaceActionType::RadioButton, FInputChord());
+	UI_COMMAND(EditorStopSendsDefaultValues, "Stop sends Default Values", "When stop is clicked in the editor, default values are sent to all patches in use", EUserInterfaceActionType::RadioButton, FInputChord());
+	UI_COMMAND(EditorStopKeepsLastValues, "Stop keeps last Values", "When stop is clicked in the editor, the last pixel mapped values are kept", EUserInterfaceActionType::RadioButton, FInputChord());
 
 	UI_COMMAND(EnableResizeMode, "Resize Mode", "Resizes components when transform handles are being dragged", EUserInterfaceActionType::RadioButton, FInputChord());
 	UI_COMMAND(EnableRotateMode, "Rotate Mode", "Rotates components when transform handles are being dragged", EUserInterfaceActionType::RadioButton, FInputChord());
@@ -32,10 +40,10 @@ void FDMXPixelMappingEditorCommands::RegisterCommands()
 	UI_COMMAND(ToggleScaleChildrenWithParent, "Scale Children with Parent", "Sets if children are scaled with their parent, when the parent is resized.", EUserInterfaceActionType::ToggleButton, FInputChord(EModifierKey::Alt, EKeys::Q));
 	UI_COMMAND(ToggleAlwaysSelectGroup, "Always select Group", "Sets if the parent is selected, if a child is clicked.", EUserInterfaceActionType::ToggleButton, FInputChord(EModifierKey::Alt, EKeys::W));
 	UI_COMMAND(ToggleApplyLayoutScriptWhenLoaded, "Apply Layout Script instantly", "Sets if layout script are applied as soon as they are loaded.", EUserInterfaceActionType::ToggleButton, FInputChord(EModifierKey::Alt, EKeys::E));
-	UI_COMMAND(ToggleShowMatrixCells, "Show Matrix Cells", "Sets if matrix cells are displayed. It is recommended that this is turned off when pixel mapping large quantities of fixtures.", EUserInterfaceActionType::ToggleButton, FInputChord(EModifierKey::Alt, EKeys::A));
-	UI_COMMAND(ToggleShowComponentNames, "Show Component Names", "Sets if the name of components are displayed. It is recommended that this is turned off when pixel mapping large quantities of fixtures.", EUserInterfaceActionType::ToggleButton, FInputChord(EModifierKey::Alt, EKeys::S));
-	UI_COMMAND(ToggleShowPatchInfo, "Show Patch Info", "Sets if information about the pach is displayed. It is recommended that this is turned off when pixel mapping large quantities of fixtures.", EUserInterfaceActionType::ToggleButton, FInputChord(EModifierKey::Alt, EKeys::D));
-	UI_COMMAND(ToggleShowCellIDs, "Show Cell IDs", "Sets if the cell IDs of matrix cells are displayed. It is recommended that this is turned off when pixel mapping large quantities of fixtures.", EUserInterfaceActionType::ToggleButton, FInputChord(EModifierKey::Alt, EKeys::F));
+	UI_COMMAND(ToggleShowMatrixCells, "Show Matrix Cells", "Sets if matrix cells are displayed. Can be turned off for better editor performance when pixel mapping large quantities of fixtures.", EUserInterfaceActionType::ToggleButton, FInputChord(EModifierKey::Alt, EKeys::A));
+	UI_COMMAND(ToggleShowComponentNames, "Show Component Names", "Sets if the name of components are displayed. Can be turned off for better editor performance when pixel mapping large quantities of fixtures.", EUserInterfaceActionType::ToggleButton, FInputChord(EModifierKey::Alt, EKeys::S));
+	UI_COMMAND(ToggleShowPatchInfo, "Show Patch Info", "Sets if information about the pach is displayed. Can be turned off for better editor performance when pixel mapping large quantities of fixtures.", EUserInterfaceActionType::ToggleButton, FInputChord(EModifierKey::Alt, EKeys::D));
+	UI_COMMAND(ToggleShowCellIDs, "Show Cell IDs", "Sets if the cell IDs of matrix cells are displayed. Can be turned off for better editor performance when pixel mapping large quantities of fixtures.", EUserInterfaceActionType::ToggleButton, FInputChord(EModifierKey::Alt, EKeys::F));
 	UI_COMMAND(ToggleShowPivot, "Show Pivot", "Sets if the pivot is displayed for selected components.", EUserInterfaceActionType::ToggleButton, FInputChord());
 }
 
diff --git a/Engine/Plugins/VirtualProduction/DMX/DMXPixelMapping/Source/DMXPixelMappingEditor/Private/DMXPixelMappingEditorCommands.h b/Engine/Plugins/VirtualProduction/DMX/DMXPixelMapping/Source/DMXPixelMappingEditor/Private/DMXPixelMappingEditorCommands.h
--- a/Engine/Plugins/VirtualProduction/DMX/DMXPixelMapping/Source/DMXPixelMappingEditor/Private/DMXPixelMappingEditorCommands.h
+++ b/Engine/Plugins/VirtualProduction/DMX/DMXPixelMapping/Source/DMXPixelMappingEditor/Private/DMXPixelMappingEditorCommands.h
@@ -14,9 +14,18 @@ public:
 
 	virtual void RegisterCommands() override;
 
+	// Toolbar related
 	TSharedPtr<FUICommandInfo> AddMapping;
 	TSharedPtr<FUICommandInfo> PlayDMX;
-	TSharedPtr<FUICommandInfo> StopPlayingDMX;
+	TSharedPtr<FUICommandInfo> PauseDMX;
+	TSharedPtr<FUICommandInfo> ResumeDMX;
+	TSharedPtr<FUICommandInfo> StopDMX;
+	TSharedPtr<FUICommandInfo> TogglePlayPauseDMX;
+	TSharedPtr<FUICommandInfo> TogglePlayStopDMX;
+
+	TSharedPtr<FUICommandInfo> EditorStopSendsDefaultValues;
+	TSharedPtr<FUICommandInfo> EditorStopSendsZeroValues;
+	TSharedPtr<FUICommandInfo> EditorStopKeepsLastValues;
 
 	// Designer related
 	TSharedPtr<FUICommandInfo> EnableResizeMode;
diff --git a/Engine/Plugins/VirtualProduction/DMX/DMXPixelMapping/Source/DMXPixelMappingEditor/Private/DMXPixelMappingEditorStyle.cpp b/Engine/Plugins/VirtualProduction/DMX/DMXPixelMapping/Source/DMXPixelMappingEditor/Private/DMXPixelMappingEditorStyle.cpp
--- a/Engine/Plugins/VirtualProduction/DMX/DMXPixelMapping/Source/DMXPixelMappingEditor/Private/DMXPixelMappingEditorStyle.cpp
+++ b/Engine/Plugins/VirtualProduction/DMX/DMXPixelMapping/Source/DMXPixelMappingEditor/Private/DMXPixelMappingEditorStyle.cpp
@@ -41,8 +41,13 @@ FDMXPixelMappingEditorStyle::FDMXPixelMappingEditorStyle()
 	{
 		Set("Icons.Preview", new IMAGE_BRUSH_SVG("Preview", Icon16x16));
 		Set("Icons.AddSource",new CORE_IMAGE_BRUSH("Icons/PlusSymbol_12x", Icon12x12));
-		Set("Icons.PlayDMX", new CORE_IMAGE_BRUSH_SVG("Starship/Common/play", Icon16x16, FStyleColors::AccentGreen));
-		Set("Icons.StopPlayingDMX", new CORE_IMAGE_BRUSH_SVG("Starship/Common/stop", Icon16x16, FStyleColors::AccentRed));
+		Set("Icons.PlayDMX", new CORE_IMAGE_BRUSH_SVG("Starship/Common/play", Icon20x20));
+		Set("Icons.StopDMX", new CORE_IMAGE_BRUSH_SVG("Starship/Common/stop", Icon20x20));
+
+		SetCoreContentRoot(FPaths::EngineContentDir() / TEXT("Editor/Slate"));
+		Set("Icons.ResumeDMX", new CORE_IMAGE_BRUSH_SVG("Starship/MainToolbar/simulate", Icon20x20));
+		Set("Icons.PauseDMX", new CORE_IMAGE_BRUSH_SVG("Starship/MainToolbar/pause", Icon20x20));
+		SetCoreContentRoot(FPaths::EngineContentDir() / TEXT("Slate"));
 	}
 
 	// Component border style
diff --git a/Engine/Plugins/VirtualProduction/DMX/DMXPixelMapping/Source/DMXPixelMappingEditor/Private/DMXPixelMappingToolbar.cpp b/Engine/Plugins/VirtualProduction/DMX/DMXPixelMapping/Source/DMXPixelMappingEditor/Private/DMXPixelMappingToolbar.cpp
--- a/Engine/Plugins/VirtualProduction/DMX/DMXPixelMapping/Source/DMXPixelMappingEditor/Private/DMXPixelMappingToolbar.cpp
+++ b/Engine/Plugins/VirtualProduction/DMX/DMXPixelMapping/Source/DMXPixelMappingEditor/Private/DMXPixelMappingToolbar.cpp
@@ -5,60 +5,127 @@
 #include "DMXPixelMappingEditorCommands.h"
 #include "DMXPixelMappingEditorStyle.h"
 #include "Framework/MultiBox/MultiBoxBuilder.h"
-#include "Framework/Commands/UICommandList.h"
-#include "Styling/AppStyle.h"
+#include "Styling/ToolBarStyle.h"
 #include "Toolkits/DMXPixelMappingToolkit.h"
+#include "Widgets/SNullWidget.h"
 
 
 #define LOCTEXT_NAMESPACE "FDMXPixelMappingToolbar"
 
 FDMXPixelMappingToolbar::FDMXPixelMappingToolbar(TSharedPtr<FDMXPixelMappingToolkit> InToolkit)
-	: ToolkitWeakPtr(InToolkit)
+	: WeakToolkit(InToolkit)
 {}
 
-void FDMXPixelMappingToolbar::BuildToolbar(TSharedPtr<FExtender> Extender)
+void FDMXPixelMappingToolbar::ExtendToolbar()
 {
-	TSharedPtr<FDMXPixelMappingToolkit> Toolkit = ToolkitWeakPtr.Pin();
-	check(Toolkit.IsValid());
-
-	Extender->AddToolBarExtension(
-		"Asset",
-		EExtensionHook::After,
-		Toolkit->GetToolkitCommands(),
-		FToolBarExtensionDelegate::CreateSP(this, &FDMXPixelMappingToolbar::BuildToolbarCallback)
-	);
-}
+	const TSharedPtr<FDMXPixelMappingToolkit> Toolkit = WeakToolkit.Pin();
+	if (!Toolkit.IsValid())
+	{
+		return;
+	}
 
-void FDMXPixelMappingToolbar::BuildToolbarCallback(FToolBarBuilder& ToolbarBuilder)
-{
-	ToolbarBuilder.BeginSection("Renderers");
+	FName ParentName;
+	const FName MenuName = Toolkit->GetToolMenuToolbarName(ParentName);
+
+	UToolMenu* ToolMenu = UToolMenus::Get()->ExtendMenu(MenuName);
+	const FToolMenuInsert SectionInsertLocation("Asset", EToolMenuInsertType::After);
+
+	// Render section
 	{
-		ToolbarBuilder.AddToolBarButton(FDMXPixelMappingEditorCommands::Get().AddMapping,
-			NAME_None, 
-			TAttribute<FText>(), 
+		FToolMenuSection& RenderSection = ToolMenu->AddSection("Render");
+
+		FToolMenuEntry AddMappingMenuEntry = FToolMenuEntry::InitToolBarButton(
+			FDMXPixelMappingEditorCommands::Get().AddMapping,
+			TAttribute<FText>(),
 			TAttribute<FText>(),
 			FSlateIcon(FDMXPixelMappingEditorStyle::Get().GetStyleSetName(), "Icons.AddSource"),
 			FName(TEXT("Add Source")));
+
+		RenderSection.AddEntry(AddMappingMenuEntry);
 	}
-	ToolbarBuilder.EndSection();
 
-	ToolbarBuilder.BeginSection("PlayAndStopDMX");
+	// Play section
 	{
-		ToolbarBuilder.AddToolBarButton(FDMXPixelMappingEditorCommands::Get().PlayDMX,
-			NAME_None, 
-			TAttribute<FText>(), 
+		FToolMenuSection& PlaySection = ToolMenu->AddSection("Play");
+
+		PlaySection.AddSeparator(NAME_None);
+
+		// Play
+		FToolMenuEntry PlayMenuEntry = FToolMenuEntry::InitToolBarButton(
+			FDMXPixelMappingEditorCommands::Get().PlayDMX, 
+			TAttribute<FText>(),
 			TAttribute<FText>(),
-			FSlateIcon(FDMXPixelMappingEditorStyle::Get().GetStyleSetName(), "Icons.PlayDMX"),
-			FName(TEXT("Play DMX")));
+			FSlateIcon(FDMXPixelMappingEditorStyle::Get().GetStyleSetName(), "Icons.PlayDMX"));
+		
+		PlayMenuEntry.StyleNameOverride = FName("Toolbar.BackplateLeftPlay");
+		PlaySection.AddEntry(PlayMenuEntry);
 
-		ToolbarBuilder.AddToolBarButton(FDMXPixelMappingEditorCommands::Get().StopPlayingDMX,
-			NAME_None, 
-			TAttribute<FText>(), 
+		// Pause
+		FToolMenuEntry PauseMenuEntry = FToolMenuEntry::InitToolBarButton(
+			FDMXPixelMappingEditorCommands::Get().PauseDMX,
 			TAttribute<FText>(),
-			FSlateIcon(FDMXPixelMappingEditorStyle::Get().GetStyleSetName(), "Icons.StopPlayingDMX"),
-			FName(TEXT("Stop Playing DMX")));
+			TAttribute<FText>(),
+			FSlateIcon(FDMXPixelMappingEditorStyle::Get().GetStyleSetName(), "Icons.PauseDMX"));
+
+		PauseMenuEntry.StyleNameOverride = FName("Toolbar.BackplateLeft");
+		PlaySection.AddEntry(PauseMenuEntry);
+
+		// Resume
+		FToolMenuEntry ResumeMenuEntry = FToolMenuEntry::InitToolBarButton(
+			FDMXPixelMappingEditorCommands::Get().ResumeDMX,
+			TAttribute<FText>(),
+			TAttribute<FText>(),
+			FSlateIcon(FDMXPixelMappingEditorStyle::Get().GetStyleSetName(), "Icons.ResumeDMX"));
+
+		ResumeMenuEntry.StyleNameOverride = FName("Toolbar.BackplateLeftPlay");
+		PlaySection.AddEntry(ResumeMenuEntry);
+
+		// Stop
+		FToolMenuEntry StopMenuEntry = FToolMenuEntry::InitToolBarButton(
+			FDMXPixelMappingEditorCommands::Get().StopDMX,
+			TAttribute<FText>(),
+			TAttribute<FText>(),
+			FSlateIcon(FDMXPixelMappingEditorStyle::Get().GetStyleSetName(), "Icons.StopDMX"));
+
+		StopMenuEntry.StyleNameOverride = FName("Toolbar.BackplateCenterStop");
+		PlaySection.AddEntry(StopMenuEntry);
+
+		// Playback Settings
+		FToolMenuEntry PlaybackSettingsComboEntry = FToolMenuEntry::InitComboButton(
+			"PlaybackSettings",
+			FToolUIActionChoice(),
+			FOnGetContent::CreateSP(this, &FDMXPixelMappingToolbar::GeneratePlaybackSettingsMenu, MenuName),
+			LOCTEXT("PlaybackSettingsLabel", "DMX Playback Settings"), 
+			LOCTEXT("PlaybackSettingsToolTip", "Change DMX Playback Settings"));
+
+		PlaybackSettingsComboEntry.StyleNameOverride = FName("Toolbar.BackplateRightCombo");
+		PlaySection.AddEntry(PlaybackSettingsComboEntry);
+
+		FToolMenuEntry BackPlateRightSeparaterEntry = PlaySection.AddSeparator(NAME_None);
+		BackPlateRightSeparaterEntry.StyleNameOverride = FName("Toolbar.BackplateRight");
 	}
-	ToolbarBuilder.EndSection();
+}
+
+TSharedRef<SWidget> FDMXPixelMappingToolbar::GeneratePlaybackSettingsMenu(FName ParentMenuName) const
+{
+	const TSharedPtr<FDMXPixelMappingToolkit> Toolkit = WeakToolkit.Pin();
+	if (!Toolkit.IsValid())
+	{
+		return SNullWidget::NullWidget;
+	}
+
+	constexpr bool bShouldCloseWindowAfterMenuSelection = true;
+	FMenuBuilder MenuBuilder(bShouldCloseWindowAfterMenuSelection, Toolkit->GetToolkitCommands());
+
+	MenuBuilder.BeginSection("ResetDMXModeSection");
+	{
+		MenuBuilder.AddMenuEntry(FDMXPixelMappingEditorCommands::Get().EditorStopSendsDefaultValues);
+		MenuBuilder.AddMenuEntry(FDMXPixelMappingEditorCommands::Get().EditorStopSendsZeroValues);
+		MenuBuilder.AddMenuEntry(FDMXPixelMappingEditorCommands::Get().EditorStopKeepsLastValues);
+	}
+	MenuBuilder.EndSection();
+
+	return MenuBuilder.MakeWidget();
 }
 
 #undef LOCTEXT_NAMESPACE
diff --git a/Engine/Plugins/VirtualProduction/DMX/DMXPixelMapping/Source/DMXPixelMappingEditor/Private/DMXPixelMappingToolbar.h b/Engine/Plugins/VirtualProduction/DMX/DMXPixelMapping/Source/DMXPixelMappingEditor/Private/DMXPixelMappingToolbar.h
--- a/Engine/Plugins/VirtualProduction/DMX/DMXPixelMapping/Source/DMXPixelMappingEditor/Private/DMXPixelMappingToolbar.h
+++ b/Engine/Plugins/VirtualProduction/DMX/DMXPixelMapping/Source/DMXPixelMappingEditor/Private/DMXPixelMappingToolbar.h
@@ -4,38 +4,24 @@
 
 #include "Templates/SharedPointer.h"
 
-class FExtender;
-class FToolBarBuilder;
 class FDMXPixelMappingToolkit;
 class SWidget;
 
-/**
- * Custom Toolbar for DMX Pixel Mapping Editor
- */
+
+/** Extends the PixelMapping asset editor toolkit toolbar */
 class FDMXPixelMappingToolbar :
 	public TSharedFromThis<FDMXPixelMappingToolbar>
 {
 public:
-	/** Default Constructor */
 	FDMXPixelMappingToolbar(TSharedPtr<FDMXPixelMappingToolkit> InToolkit);
-
-	/** Virtual Destructor */
 	virtual ~FDMXPixelMappingToolbar() {}
 
-	/** Builds the toolbar */
-	void BuildToolbar(TSharedPtr<FExtender> Extender);
+	/** Extends the PixelMapping asset editor toolkit toolbar  */
+	void ExtendToolbar();
 
 private:
-	/** Callback, raised when the menu extender requests to build the toolbar */
-	void BuildToolbarCallback(FToolBarBuilder& ToolbarBuilder);
-
-	/** Generates a widget with play options (play, stop) */
-	TSharedRef<SWidget> GeneratesPlayOptionsWidget();
+	/** Generates the playback settings submenu */
+	TSharedRef<SWidget> GeneratePlaybackSettingsMenu(FName ParentMenuName) const;
 
-	/** Generates a layout settings menu for the pixelmapping toolbar */
-	TSharedRef<SWidget> GenerateLayoutMenu();
-
-public:
-	TWeakPtr<FDMXPixelMappingToolkit> ToolkitWeakPtr;
+	TWeakPtr<FDMXPixelMappingToolkit> WeakToolkit;
 };
-
diff --git a/Engine/Plugins/VirtualProduction/DMX/DMXPixelMapping/Source/DMXPixelMappingEditor/Private/Settings/DMXPixelMappingEditorSettings.h b/Engine/Plugins/VirtualProduction/DMX/DMXPixelMapping/Source/DMXPixelMappingEditor/Private/Settings/DMXPixelMappingEditorSettings.h
--- a/Engine/Plugins/VirtualProduction/DMX/DMXPixelMapping/Source/DMXPixelMappingEditor/Private/Settings/DMXPixelMappingEditorSettings.h
+++ b/Engine/Plugins/VirtualProduction/DMX/DMXPixelMapping/Source/DMXPixelMappingEditor/Private/Settings/DMXPixelMappingEditorSettings.h
@@ -4,6 +4,7 @@
 
 #include "DMXPixelMappingEditorSettings.generated.h"
 
+enum class EDMXPixelMappingResetDMXMode : uint8;
 struct FPropertyChangedEvent;
 
 
@@ -88,6 +89,10 @@ public:
 	UPROPERTY(Config)
 	FDMXPixelMappingDesignerSettings DesignerSettings;
 
+	/** Reset DMX mode to be used in editor */
+	UPROPERTY(Config)
+	EDMXPixelMappingResetDMXMode EditorResetDMXMode;
+
 	/** Raised by the pixel mapping toolkit when settings changed */
 	static FSimpleMulticastDelegate OnEditorSettingsChanged;
 };
diff --git a/Engine/Plugins/VirtualProduction/DMX/DMXPixelMapping/Source/DMXPixelMappingEditor/Private/Toolkits/DMXPixelMappingToolkit.cpp b/Engine/Plugins/VirtualProduction/DMX/DMXPixelMapping/Source/DMXPixelMappingEditor/Private/Toolkits/DMXPixelMappingToolkit.cpp
--- a/Engine/Plugins/VirtualProduction/DMX/DMXPixelMapping/Source/DMXPixelMappingEditor/Private/Toolkits/DMXPixelMappingToolkit.cpp
+++ b/Engine/Plugins/VirtualProduction/DMX/DMXPixelMapping/Source/DMXPixelMappingEditor/Private/Toolkits/DMXPixelMappingToolkit.cpp
@@ -10,6 +10,7 @@
 #include "DMXPixelMappingEditorStyle.h"
 #include "DMXPixelMappingEditorUtils.h"
 #include "DMXPixelMappingToolbar.h"
+#include "Engine/Texture.h"
 #include "Framework/Commands/GenericCommands.h"
 #include "Framework/MultiBox/MultiBoxExtender.h"
 #include "K2Node_PixelMappingBaseComponent.h"
@@ -252,26 +253,13 @@ void FDMXPixelMappingToolkit::Tick(float DeltaTime)
 		return;
 	}
 
-	// Render Output and Send DMX if required
-	if (!bRequestStopSendingDMX)
+	// Render, send DMX if required
+	RootComponent->Render();
+	if (bIsPlayingDMX && !bIsPaused)
 	{
-		if (bIsPlayingDMX)
-		{
-			RootComponent->RenderAndSendDMX();
-		}
-		else
-		{
-			RootComponent->Render();
-		}
-	}
-	else if (bRequestStopSendingDMX)
-	{
-		RootComponent->ResetDMX();
-			
-		bIsPlayingDMX = false;
-		bRequestStopSendingDMX = false; 
+		RootComponent->SendDMX();
 	}
-	
+
 	// Detect and broadcast editor setting changes
 	if (FMemory::Memcmp(EditorSettingsDump.GetData(), GetDefault<UDMXPixelMappingEditorSettings>(), sizeof(UDMXPixelMappingEditorSettings)) != 0)
 	{
@@ -403,12 +391,66 @@ void FDMXPixelMappingToolkit::AddRenderer()
 
 void FDMXPixelMappingToolkit::PlayDMX()
 {
+	bIsPaused = false;
 	bIsPlayingDMX = true;
 }
 
+void FDMXPixelMappingToolkit::PauseDMX()
+{
+	bIsPaused = true;
+	bIsPlayingDMX = false;
+}
+
 void FDMXPixelMappingToolkit::StopPlayingDMX()
 {
-	bRequestStopSendingDMX = true;
+	UDMXPixelMapping* PixelMapping = GetDMXPixelMapping();
+	if (!PixelMapping)
+	{
+		return;
+	}
+
+	UDMXPixelMappingRootComponent* RootComponent = PixelMapping->RootComponent;
+	if (!ensure(RootComponent))
+	{
+		return;
+	}
+
+	const UDMXPixelMappingEditorSettings* Settings = GetDefault<UDMXPixelMappingEditorSettings>();
+	RootComponent->ResetDMX(Settings->EditorResetDMXMode);
+
+	bIsPaused = false;
+	bIsPlayingDMX = false;
+}
+
+void FDMXPixelMappingToolkit::TogglePlayPauseDMX()
+{
+	if (bIsPlayingDMX)
+	{
+		PauseDMX();
+	}
+	else
+	{
+		PlayDMX();
+	}
+}
+
+void FDMXPixelMappingToolkit::TogglePlayStopDMX()
+{
+	if (bIsPlayingDMX)
+	{
+		StopPlayingDMX();
+	}
+	else
+	{
+		PlayDMX();
+	}
+}
+
+void FDMXPixelMappingToolkit::SetEditorResetDMXMode(EDMXPixelMappingResetDMXMode NewMode)
+{
+	UDMXPixelMappingEditorSettings* EditorSettings = GetMutableDefault<UDMXPixelMappingEditorSettings>();
+	EditorSettings->EditorResetDMXMode = NewMode;
+	EditorSettings->SaveConfig();
 }
 
 void FDMXPixelMappingToolkit::UpdateBlueprintNodes() const
@@ -425,9 +467,13 @@ void FDMXPixelMappingToolkit::UpdateBlueprintNodes() const
 void FDMXPixelMappingToolkit::SaveThumbnailImage()
 {
 	UDMXPixelMapping* PixelMapping = GetDMXPixelMapping();
-	if (PixelMapping && ActiveRendererComponent.IsValid())
+	if (PixelMapping)
 	{
-		PixelMapping->ThumbnailImage = ActiveRendererComponent->GetRenderedInputTexture();
+		UTexture* Texture = ActiveRendererComponent.IsValid() ? ActiveRendererComponent->GetRenderedInputTexture() : nullptr;
+		if (IsValid(Texture) && Texture->IsFullyStreamedIn())
+		{
+			PixelMapping->ThumbnailImage = Texture;
+		}
 	}
 }
 
@@ -849,31 +895,85 @@ void FDMXPixelMappingToolkit::SetupCommands()
 	GetToolkitCommands()->MapAction(
 		FDMXPixelMappingEditorCommands::Get().PlayDMX,
 		FExecuteAction::CreateSP(this, &FDMXPixelMappingToolkit::PlayDMX),
-		FCanExecuteAction::CreateLambda([this] 
-			{ 
-				return !bIsPlayingDMX;
+		FCanExecuteAction::CreateLambda([this]
+			{
+				return !bIsPlayingDMX && !bIsPaused;
 			}),
 		FIsActionChecked(),
-		FIsActionButtonVisible::CreateLambda([this] 
-			{ 
-				return !bIsPlayingDMX; 
+		FIsActionButtonVisible::CreateLambda([this]
+			{
+				return !bIsPlayingDMX && !bIsPaused;
 			})
 	);
 
 	GetToolkitCommands()->MapAction(
-		FDMXPixelMappingEditorCommands::Get().StopPlayingDMX,
-		FExecuteAction::CreateSP(this, &FDMXPixelMappingToolkit::StopPlayingDMX),
+		FDMXPixelMappingEditorCommands::Get().PauseDMX,
+		FExecuteAction::CreateSP(this, &FDMXPixelMappingToolkit::PauseDMX),
 		FCanExecuteAction::CreateLambda([this] 
 			{ 
 				return bIsPlayingDMX; 
 			}),
 		FIsActionChecked(),
-		FIsActionButtonVisible::CreateLambda([this] 
-			{ 
+		FIsActionButtonVisible::CreateLambda([this]
+			{
 				return bIsPlayingDMX;
 			})
 	);
 
+	GetToolkitCommands()->MapAction(
+		FDMXPixelMappingEditorCommands::Get().ResumeDMX,
+		FExecuteAction::CreateSP(this, &FDMXPixelMappingToolkit::PlayDMX),
+		FCanExecuteAction::CreateLambda([this]
+			{
+				return !bIsPlayingDMX && bIsPaused;
+			}),
+		FIsActionChecked(),
+		FIsActionButtonVisible::CreateLambda([this]
+			{
+				return !bIsPlayingDMX && bIsPaused;
+			})
+	);
+
+	GetToolkitCommands()->MapAction(
+		FDMXPixelMappingEditorCommands::Get().StopDMX,
+		FExecuteAction::CreateSP(this, &FDMXPixelMappingToolkit::StopPlayingDMX),
+		FCanExecuteAction::CreateLambda([this]
+			{
+				return bIsPlayingDMX || bIsPaused;
+			})
+	);
+
+	GetToolkitCommands()->MapAction(
+		FDMXPixelMappingEditorCommands::Get().TogglePlayPauseDMX,
+		FExecuteAction::CreateSP(this, &FDMXPixelMappingToolkit::TogglePlayPauseDMX)
+	);
+
+	GetToolkitCommands()->MapAction(
+		FDMXPixelMappingEditorCommands::Get().TogglePlayStopDMX,
+		FExecuteAction::CreateSP(this, &FDMXPixelMappingToolkit::TogglePlayStopDMX)
+		);
+
+	GetToolkitCommands()->MapAction(
+		FDMXPixelMappingEditorCommands::Get().EditorStopSendsDefaultValues,
+		FExecuteAction::CreateSP(this, &FDMXPixelMappingToolkit::SetEditorResetDMXMode, EDMXPixelMappingResetDMXMode::SendDefaultValues),
+		FCanExecuteAction(),
+		FGetActionCheckState::CreateSP(this, &FDMXPixelMappingToolkit::GetEditorResetDMXModeCheckboxState, EDMXPixelMappingResetDMXMode::SendDefaultValues)
+	);
+
+	GetToolkitCommands()->MapAction(
+		FDMXPixelMappingEditorCommands::Get().EditorStopSendsZeroValues,
+		FExecuteAction::CreateSP(this, &FDMXPixelMappingToolkit::SetEditorResetDMXMode, EDMXPixelMappingResetDMXMode::SendZeroValues),
+		FCanExecuteAction(),
+		FGetActionCheckState::CreateSP(this, &FDMXPixelMappingToolkit::GetEditorResetDMXModeCheckboxState, EDMXPixelMappingResetDMXMode::SendZeroValues)
+	);
+
+	GetToolkitCommands()->MapAction(
+		FDMXPixelMappingEditorCommands::Get().EditorStopKeepsLastValues,
+		FExecuteAction::CreateSP(this, &FDMXPixelMappingToolkit::SetEditorResetDMXMode, EDMXPixelMappingResetDMXMode::DoNotSendValues),
+		FCanExecuteAction(),
+		FGetActionCheckState::CreateSP(this, &FDMXPixelMappingToolkit::GetEditorResetDMXModeCheckboxState, EDMXPixelMappingResetDMXMode::DoNotSendValues)
+	);
+
 	GetToolkitCommands()->MapAction(
 		FDMXPixelMappingEditorCommands::Get().EnableResizeMode,
 		FExecuteAction::CreateSP(this, &FDMXPixelMappingToolkit::SetTransformHandleMode, EDMXPixelMappingTransformHandleMode::Resize),
@@ -915,18 +1015,22 @@ void FDMXPixelMappingToolkit::SetupCommands()
 
 void FDMXPixelMappingToolkit::ExtendToolbar()
 {
-	FDMXPixelMappingEditorModule& DMXPixelMappingEditorModule = FModuleManager::LoadModuleChecked<FDMXPixelMappingEditorModule>("DMXPixelMappingEditor");
 	Toolbar = MakeShared<FDMXPixelMappingToolbar>(SharedThis(this));
 
-	TSharedPtr<FExtender> ToolbarExtender = MakeShareable(new FExtender);
-	Toolbar->BuildToolbar(ToolbarExtender);
-	AddToolbarExtender(ToolbarExtender);
+	Toolbar->ExtendToolbar();
 
 	// Let other part of the plugin extend DMX Pixel Maping Editor toolbar
+	FDMXPixelMappingEditorModule& DMXPixelMappingEditorModule = FModuleManager::LoadModuleChecked<FDMXPixelMappingEditorModule>("DMXPixelMappingEditor");
 	AddMenuExtender(DMXPixelMappingEditorModule.GetMenuExtensibilityManager()->GetAllExtenders(GetToolkitCommands(), GetEditingObjects()));
 	AddToolbarExtender(DMXPixelMappingEditorModule.GetToolBarExtensibilityManager()->GetAllExtenders(GetToolkitCommands(), GetEditingObjects()));
 }
 
+ECheckBoxState FDMXPixelMappingToolkit::GetEditorResetDMXModeCheckboxState(EDMXPixelMappingResetDMXMode CompareMode) const
+{
+	const UDMXPixelMappingEditorSettings* EditorSettings = GetDefault<UDMXPixelMappingEditorSettings>();
+	return EditorSettings->EditorResetDMXMode == CompareMode ? ECheckBoxState::Checked : ECheckBoxState::Unchecked;
+}
+
 ECheckBoxState FDMXPixelMappingToolkit::GetTransformHandleModeCheckboxState(EDMXPixelMappingTransformHandleMode CompareTransformHandleMode) const
 {
 	return CompareTransformHandleMode == TransformHandleMode ? ECheckBoxState::Checked : ECheckBoxState::Unchecked;
diff --git a/Engine/Plugins/VirtualProduction/DMX/DMXPixelMapping/Source/DMXPixelMappingEditor/Private/Toolkits/DMXPixelMappingToolkit.h b/Engine/Plugins/VirtualProduction/DMX/DMXPixelMapping/Source/DMXPixelMappingEditor/Private/Toolkits/DMXPixelMappingToolkit.h
--- a/Engine/Plugins/VirtualProduction/DMX/DMXPixelMapping/Source/DMXPixelMappingEditor/Private/Toolkits/DMXPixelMappingToolkit.h
+++ b/Engine/Plugins/VirtualProduction/DMX/DMXPixelMapping/Source/DMXPixelMappingEditor/Private/Toolkits/DMXPixelMappingToolkit.h
@@ -10,6 +10,8 @@
 #include "UObject/GCObject.h"
 #include "Widgets/Views/SHeaderRow.h"
 
+enum class ECheckBoxState : uint8;
+enum class EDMXPixelMappingResetDMXMode : uint8;
 class FDMXPixelMappingComponentTemplate;
 class FDMXPixelMappingToolbar;
 class FSpawnTabArgs;
@@ -26,7 +28,6 @@ class UDMXPixelMappingBaseComponent;
 class UDMXPixelMappingMatrixComponent;
 class UDMXPixelMappingOutputComponent;
 class UDMXPixelMappingRendererComponent;
-enum class ECheckBoxState : uint8;
 
 namespace UE::DMX
 {
@@ -172,10 +173,24 @@ private:
 	/** Called when a component object was renamed */
 	void OnComponentRenamed(UDMXPixelMappingBaseComponent* Component);
 
+	/** Starts to play DMX on editor tick */
 	void PlayDMX();
 
+	/** Pauses playing DMX. Current DMX values will still be sent at a lower rate. */
+	void PauseDMX();
+
+	/** Stops playing DMX */
 	void StopPlayingDMX();
 
+	/** Toggles between playing and pausing DMX */
+	void TogglePlayPauseDMX();
+
+	/** Toggles between playing and stopping DMX */
+	void TogglePlayStopDMX();
+
+	/** Sets the reset DMX mode used by the editor */
+	void SetEditorResetDMXMode(EDMXPixelMappingResetDMXMode NewMode);
+
 	/** Updates blueprint nodes */
 	void UpdateBlueprintNodes() const;
 
@@ -206,6 +221,9 @@ private:
 
 	void CreateInternalViews();
 
+	/** Returns the check box state for the compared reset DMX mode */
+	ECheckBoxState GetEditorResetDMXModeCheckboxState(EDMXPixelMappingResetDMXMode CompareMode) const;
+
 	/** Returns the checkbox state for a transform handle mode, checked if the mode equals the current mode. */
 	ECheckBoxState GetTransformHandleModeCheckboxState(EDMXPixelMappingTransformHandleMode CompareTransformHandleMode) const;
 
@@ -246,11 +264,8 @@ private:
 	/** True while playing DMX */
 	bool bIsPlayingDMX = false;
 
-	/** Toggles if DMX should be sent for all components */
-	bool bTogglePlayDMXAll = true;
-
-	/** True while stop DMX is requested, but not carried out yet */
-	bool bRequestStopSendingDMX = false;
+	/** True if paused */
+	bool bIsPaused = false;
 
 	/** True while adding components (to avoid needlessly updating blueprint nodes on each component added via our own methods) */
 	bool bAddingComponents = false;
diff --git a/Engine/Plugins/VirtualProduction/DMX/DMXPixelMapping/Source/DMXPixelMappingRuntime/Private/Components/DMXPixelMappingBaseComponent.cpp b/Engine/Plugins/VirtualProduction/DMX/DMXPixelMapping/Source/DMXPixelMappingRuntime/Private/Components/DMXPixelMappingBaseComponent.cpp
--- a/Engine/Plugins/VirtualProduction/DMX/DMXPixelMapping/Source/DMXPixelMappingRuntime/Private/Components/DMXPixelMappingBaseComponent.cpp
+++ b/Engine/Plugins/VirtualProduction/DMX/DMXPixelMapping/Source/DMXPixelMappingRuntime/Private/Components/DMXPixelMappingBaseComponent.cpp
@@ -204,7 +204,7 @@ void UDMXPixelMappingBaseComponent::RemoveChild(UDMXPixelMappingBaseComponent* C
 {
 	if (ChildComponent)
 	{
-		ChildComponent->ResetDMX();
+		ChildComponent->ResetDMX(EDMXPixelMappingResetDMXMode::SendDefaultValues);
 
 		Children.Remove(ChildComponent);
 		ChildComponent->NotifyRemovedFromParent();
diff --git a/Engine/Plugins/VirtualProduction/DMX/DMXPixelMapping/Source/DMXPixelMappingRuntime/Private/Components/DMXPixelMappingFixtureGroupComponent.cpp b/Engine/Plugins/VirtualProduction/DMX/DMXPixelMapping/Source/DMXPixelMappingRuntime/Private/Components/DMXPixelMappingFixtureGroupComponent.cpp
--- a/Engine/Plugins/VirtualProduction/DMX/DMXPixelMapping/Source/DMXPixelMappingRuntime/Private/Components/DMXPixelMappingFixtureGroupComponent.cpp
+++ b/Engine/Plugins/VirtualProduction/DMX/DMXPixelMapping/Source/DMXPixelMappingRuntime/Private/Components/DMXPixelMappingFixtureGroupComponent.cpp
@@ -81,13 +81,13 @@ void UDMXPixelMappingFixtureGroupComponent::AddChild(UDMXPixelMappingBaseCompone
 	}
 }
 
-void UDMXPixelMappingFixtureGroupComponent::ResetDMX()
+void UDMXPixelMappingFixtureGroupComponent::ResetDMX(EDMXPixelMappingResetDMXMode ResetMode)
 {
 	ForEachChild([&](UDMXPixelMappingBaseComponent* InComponent)
 	{
 		if (UDMXPixelMappingOutputComponent * Component = Cast<UDMXPixelMappingOutputComponent>(InComponent))
 		{
-			Component->ResetDMX();
+			Component->ResetDMX(ResetMode);
 		}
 	}, false);
 }
diff --git a/Engine/Plugins/VirtualProduction/DMX/DMXPixelMapping/Source/DMXPixelMappingRuntime/Private/Components/DMXPixelMappingFixtureGroupItemComponent.cpp b/Engine/Plugins/VirtualProduction/DMX/DMXPixelMapping/Source/DMXPixelMappingRuntime/Private/Components/DMXPixelMappingFixtureGroupItemComponent.cpp
--- a/Engine/Plugins/VirtualProduction/DMX/DMXPixelMapping/Source/DMXPixelMappingRuntime/Private/Components/DMXPixelMappingFixtureGroupItemComponent.cpp
+++ b/Engine/Plugins/VirtualProduction/DMX/DMXPixelMapping/Source/DMXPixelMappingRuntime/Private/Components/DMXPixelMappingFixtureGroupItemComponent.cpp
@@ -169,7 +169,7 @@ void UDMXPixelMappingFixtureGroupItemComponent::PostEditChangeProperty(FProperty
 			if (ColorSpace)
 			{
 				ColorSpace->GetOnPostEditChangedProperty().RemoveAll(this);
-				ResetDMX();
+				ResetDMX(EDMXPixelMappingResetDMXMode::DoNotSendValues);
 			}
 
 			ColorSpace = NewObject<UDMXPixelMappingColorSpace>(this, ColorSpaceClass);
@@ -225,39 +225,20 @@ bool UDMXPixelMappingFixtureGroupItemComponent::IsVisible() const
 }
 #endif // WITH_EDITOR
 
-void UDMXPixelMappingFixtureGroupItemComponent::ResetDMX()
+void UDMXPixelMappingFixtureGroupItemComponent::ResetDMX(EDMXPixelMappingResetDMXMode ResetMode)
 {
-	UDMXPixelMappingRendererComponent* RendererComponent = GetRendererComponent();
-	if (!ensure(RendererComponent))
-	{
-		return;
-	}
-
 	UDMXEntityFixturePatch* FixturePatch = FixturePatchRef.GetFixturePatch();
-	if (!FixturePatch)
-	{
-		return;
-	}
 
-	UDMXLibrary* Library = FixturePatch->GetParentLibrary();
-	if (!Library)
+	if (FixturePatch)
 	{
-		return;
-	}
-	
-	ColorSpace->ResetToBlack();
-	TMap<FDMXAttributeName, float> AttributeToValueMap = ColorSpace->GetAttributeNameToValueMap();
-
-	TMap<int32, uint8> ChannelToValueMap;
-	for (const TTuple<FDMXAttributeName, float>& AttributeValuePair : AttributeToValueMap)
-	{
-		FDMXPixelMappingRuntimeUtils::ConvertNormalizedAttributeValueToChannelValue(FixturePatch, AttributeValuePair.Key, AttributeValuePair.Value, ChannelToValueMap);
-	}
-
-	// Send DMX
-	for (const FDMXOutputPortSharedRef& OutputPort : Library->GetOutputPorts())
-	{
-		OutputPort->SendDMX(FixturePatch->GetUniverseID(), ChannelToValueMap);
+		if (ResetMode == EDMXPixelMappingResetDMXMode::SendZeroValues)
+		{
+			FixturePatch->SendZeroValues();
+		}
+		else if (ResetMode == EDMXPixelMappingResetDMXMode::SendDefaultValues)
+		{
+			FixturePatch->SendDefaultValues();
+		}
 	}
 }
 
@@ -453,7 +434,7 @@ void UDMXPixelMappingFixtureGroupItemComponent::OnColorSpacePostEditChangeProper
 	if (PropertyChangedEvent.GetPropertyName() == GET_MEMBER_NAME_CHECKED(FDMXAttributeName, Name))
 	{
 		// Reset DMX when an attribute of the color space changed
-		ResetDMX();
+		ResetDMX(EDMXPixelMappingResetDMXMode::DoNotSendValues);
 	}
 }
 #endif // WITH_EDITOR
diff --git a/Engine/Plugins/VirtualProduction/DMX/DMXPixelMapping/Source/DMXPixelMappingRuntime/Private/Components/DMXPixelMappingMatrixCellComponent.cpp b/Engine/Plugins/VirtualProduction/DMX/DMXPixelMapping/Source/DMXPixelMappingRuntime/Private/Components/DMXPixelMappingMatrixCellComponent.cpp
--- a/Engine/Plugins/VirtualProduction/DMX/DMXPixelMapping/Source/DMXPixelMappingRuntime/Private/Components/DMXPixelMappingMatrixCellComponent.cpp
+++ b/Engine/Plugins/VirtualProduction/DMX/DMXPixelMapping/Source/DMXPixelMappingRuntime/Private/Components/DMXPixelMappingMatrixCellComponent.cpp
@@ -102,20 +102,6 @@ const FName& UDMXPixelMappingMatrixCellComponent::GetNamePrefix()
 	return NamePrefix;
 }
 
-void UDMXPixelMappingMatrixCellComponent::ResetDMX()
-{
-	UDMXPixelMappingMatrixComponent* ParentMatrix = Cast<UDMXPixelMappingMatrixComponent>(GetParent());
-	if (!ParentMatrix)
-	{
-		return;
-	}
-
-	UDMXPixelMappingColorSpace* ColorSpace = ParentMatrix->ColorSpace;
-	ColorSpace->SetRGBA(FLinearColor::Black);
-
-	// No need to send dmx, that is done by the parent matrix
-}
-
 FString UDMXPixelMappingMatrixCellComponent::GetUserName() const
 {
 	UDMXPixelMappingMatrixComponent* MatrixComponent = Cast<UDMXPixelMappingMatrixComponent>(GetParent());
diff --git a/Engine/Plugins/VirtualProduction/DMX/DMXPixelMapping/Source/DMXPixelMappingRuntime/Private/Components/DMXPixelMappingMatrixComponent.cpp b/Engine/Plugins/VirtualProduction/DMX/DMXPixelMapping/Source/DMXPixelMappingRuntime/Private/Components/DMXPixelMappingMatrixComponent.cpp
--- a/Engine/Plugins/VirtualProduction/DMX/DMXPixelMapping/Source/DMXPixelMappingRuntime/Private/Components/DMXPixelMappingMatrixComponent.cpp
+++ b/Engine/Plugins/VirtualProduction/DMX/DMXPixelMapping/Source/DMXPixelMappingRuntime/Private/Components/DMXPixelMappingMatrixComponent.cpp
@@ -237,17 +237,21 @@ const FName& UDMXPixelMappingMatrixComponent::GetNamePrefix()
 	return NamePrefix;
 }
 
-void UDMXPixelMappingMatrixComponent::ResetDMX()
+void UDMXPixelMappingMatrixComponent::ResetDMX(EDMXPixelMappingResetDMXMode ResetMode)
 {
-	ForEachChild([&](UDMXPixelMappingBaseComponent* InComponent)
+	UDMXEntityFixturePatch* FixturePatch = FixturePatchRef.GetFixturePatch();
+
+	if (FixturePatch)
 	{
-		if (UDMXPixelMappingOutputComponent* Component = Cast<UDMXPixelMappingOutputComponent>(InComponent))
+		if (ResetMode == EDMXPixelMappingResetDMXMode::SendZeroValues)
 		{
-			Component->ResetDMX();
+			FixturePatch->SendZeroValues();
 		}
-	}, false);
-
-	SendDMX();
+		else if (ResetMode == EDMXPixelMappingResetDMXMode::SendDefaultValues)
+		{
+			FixturePatch->SendDefaultValues();
+		}
+	}
 }
 
 void UDMXPixelMappingMatrixComponent::SendDMX()
diff --git a/Engine/Plugins/VirtualProduction/DMX/DMXPixelMapping/Source/DMXPixelMappingRuntime/Private/Components/DMXPixelMappingRendererComponent.cpp b/Engine/Plugins/VirtualProduction/DMX/DMXPixelMapping/Source/DMXPixelMappingRuntime/Private/Components/DMXPixelMappingRendererComponent.cpp
--- a/Engine/Plugins/VirtualProduction/DMX/DMXPixelMapping/Source/DMXPixelMappingRuntime/Private/Components/DMXPixelMappingRendererComponent.cpp
+++ b/Engine/Plugins/VirtualProduction/DMX/DMXPixelMapping/Source/DMXPixelMappingRuntime/Private/Components/DMXPixelMappingRendererComponent.cpp
@@ -217,13 +217,13 @@ TArray<TSharedRef<UE::DMXPixelMapping::Rendering::FPixelMapRenderElement>> UDMXP
 	return PixelMapRenderElements;
 }
 
-void UDMXPixelMappingRendererComponent::ResetDMX()
+void UDMXPixelMappingRendererComponent::ResetDMX(EDMXPixelMappingResetDMXMode ResetMode)
 {
 	ForEachChild([&](UDMXPixelMappingBaseComponent* InComponent)
 		{
 			if (UDMXPixelMappingOutputComponent* Component = Cast<UDMXPixelMappingOutputComponent>(InComponent))
 			{
-				Component->ResetDMX();
+				Component->ResetDMX(ResetMode);
 			}
 		}, false);
 }
diff --git a/Engine/Plugins/VirtualProduction/DMX/DMXPixelMapping/Source/DMXPixelMappingRuntime/Private/Components/DMXPixelMappingRootComponent.cpp b/Engine/Plugins/VirtualProduction/DMX/DMXPixelMapping/Source/DMXPixelMappingRuntime/Private/Components/DMXPixelMappingRootComponent.cpp
--- a/Engine/Plugins/VirtualProduction/DMX/DMXPixelMapping/Source/DMXPixelMappingRuntime/Private/Components/DMXPixelMappingRootComponent.cpp
+++ b/Engine/Plugins/VirtualProduction/DMX/DMXPixelMapping/Source/DMXPixelMappingRuntime/Private/Components/DMXPixelMappingRootComponent.cpp
@@ -47,13 +47,13 @@ void UDMXPixelMappingRootComponent::RemoveChild(UDMXPixelMappingBaseComponent* I
 	CachedRendererComponentsByName.Remove(InComponent->GetFName());
 }
 
-void UDMXPixelMappingRootComponent::ResetDMX()
+void UDMXPixelMappingRootComponent::ResetDMX(EDMXPixelMappingResetDMXMode ResetMode)
 {
 	ForEachChild([&](UDMXPixelMappingBaseComponent* InComponent)
 	{
 		if (UDMXPixelMappingRendererComponent* Component = Cast<UDMXPixelMappingRendererComponent>(InComponent))
 		{
-			Component->ResetDMX();
+			Component->ResetDMX(ResetMode);
 		}
 	}, false);
 }
diff --git a/Engine/Plugins/VirtualProduction/DMX/DMXPixelMapping/Source/DMXPixelMappingRuntime/Private/Components/DMXPixelMappingScreenComponent.cpp b/Engine/Plugins/VirtualProduction/DMX/DMXPixelMapping/Source/DMXPixelMappingRuntime/Private/Components/DMXPixelMappingScreenComponent.cpp
--- a/Engine/Plugins/VirtualProduction/DMX/DMXPixelMapping/Source/DMXPixelMappingRuntime/Private/Components/DMXPixelMappingScreenComponent.cpp
+++ b/Engine/Plugins/VirtualProduction/DMX/DMXPixelMapping/Source/DMXPixelMappingRuntime/Private/Components/DMXPixelMappingScreenComponent.cpp
@@ -214,11 +214,6 @@ UDMXPixelMappingRendererComponent* UDMXPixelMappingScreenComponent::GetRendererC
 	return Cast<UDMXPixelMappingRendererComponent>(GetParent());
 }
 
-void UDMXPixelMappingScreenComponent::ResetDMX()
-{
-	SendDMX();
-}
-
 void UDMXPixelMappingScreenComponent::SendDMX()
 {
 	PRAGMA_DISABLE_DEPRECATION_WARNINGS
diff --git a/Engine/Plugins/VirtualProduction/DMX/DMXPixelMapping/Source/DMXPixelMappingRuntime/Public/Components/DMXPixelMappingBaseComponent.h b/Engine/Plugins/VirtualProduction/DMX/DMXPixelMapping/Source/DMXPixelMappingRuntime/Public/Components/DMXPixelMappingBaseComponent.h
--- a/Engine/Plugins/VirtualProduction/DMX/DMXPixelMapping/Source/DMXPixelMappingRuntime/Public/Components/DMXPixelMappingBaseComponent.h
+++ b/Engine/Plugins/VirtualProduction/DMX/DMXPixelMapping/Source/DMXPixelMappingRuntime/Public/Components/DMXPixelMappingBaseComponent.h
@@ -15,6 +15,14 @@ class UDMXPixelMapping;
 class UDMXPixelMappingRootComponent;
 
 
+UENUM(BlueprintType)
+enum class EDMXPixelMappingResetDMXMode : uint8
+{
+	SendDefaultValues UMETA(DisplayName = "Send Default Values"),
+	SendZeroValues UMETA(DisplayName = "Send Zero Values"),
+	DoNotSendValues UMETA(DisplayName = "Keep Last Mapped Values")
+};
+
 /**
  * Base class for all DMX Pixel Mapping components. 
  */
@@ -172,7 +180,7 @@ public:
 
 	/** Reset all sending DMX channels to 0 for this component and all children */
 	UFUNCTION(BlueprintCallable, Category = "DMX|PixelMapping")
-	virtual void ResetDMX() {};
+	virtual void ResetDMX(EDMXPixelMappingResetDMXMode ResetMode = EDMXPixelMappingResetDMXMode::SendDefaultValues) {};
 
 	/** Send DMX values of this component and all children. */
 	UFUNCTION(BlueprintCallable, Category = "DMX|PixelMapping")
diff --git a/Engine/Plugins/VirtualProduction/DMX/DMXPixelMapping/Source/DMXPixelMappingRuntime/Public/Components/DMXPixelMappingFixtureGroupComponent.h b/Engine/Plugins/VirtualProduction/DMX/DMXPixelMapping/Source/DMXPixelMappingRuntime/Public/Components/DMXPixelMappingFixtureGroupComponent.h
--- a/Engine/Plugins/VirtualProduction/DMX/DMXPixelMapping/Source/DMXPixelMappingRuntime/Public/Components/DMXPixelMappingFixtureGroupComponent.h
+++ b/Engine/Plugins/VirtualProduction/DMX/DMXPixelMapping/Source/DMXPixelMappingRuntime/Public/Components/DMXPixelMappingFixtureGroupComponent.h
@@ -36,7 +36,7 @@ public:
 	virtual const FName& GetNamePrefix() override;
 	virtual void AddChild(UDMXPixelMappingBaseComponent* InComponent) override;
 	virtual bool CanBeMovedTo(const UDMXPixelMappingBaseComponent* Component) const override;
-	virtual void ResetDMX() override;
+	virtual void ResetDMX(EDMXPixelMappingResetDMXMode ResetMode = EDMXPixelMappingResetDMXMode::SendDefaultValues) override;
 	virtual void SendDMX() override;
 	virtual FString GetUserName() const override;
 	//~ End UDMXPixelMappingBaseComponent implementation
diff --git a/Engine/Plugins/VirtualProduction/DMX/DMXPixelMapping/Source/DMXPixelMappingRuntime/Public/Components/DMXPixelMappingFixtureGroupItemComponent.h b/Engine/Plugins/VirtualProduction/DMX/DMXPixelMapping/Source/DMXPixelMappingRuntime/Public/Components/DMXPixelMappingFixtureGroupItemComponent.h
--- a/Engine/Plugins/VirtualProduction/DMX/DMXPixelMapping/Source/DMXPixelMappingRuntime/Public/Components/DMXPixelMappingFixtureGroupItemComponent.h
+++ b/Engine/Plugins/VirtualProduction/DMX/DMXPixelMapping/Source/DMXPixelMappingRuntime/Public/Components/DMXPixelMappingFixtureGroupItemComponent.h
@@ -49,7 +49,7 @@ protected:
 public:
 	//~ Begin UDMXPixelMappingBaseComponent implementation
 	virtual const FName& GetNamePrefix() override;
-	virtual void ResetDMX() override;
+	virtual void ResetDMX(EDMXPixelMappingResetDMXMode ResetMode = EDMXPixelMappingResetDMXMode::SendDefaultValues) override;
 	virtual void SendDMX() override;
 	virtual FString GetUserName() const override;
 	//~ End UDMXPixelMappingBaseComponent implementation
diff --git a/Engine/Plugins/VirtualProduction/DMX/DMXPixelMapping/Source/DMXPixelMappingRuntime/Public/Components/DMXPixelMappingMatrixCellComponent.h b/Engine/Plugins/VirtualProduction/DMX/DMXPixelMapping/Source/DMXPixelMappingRuntime/Public/Components/DMXPixelMappingMatrixCellComponent.h
--- a/Engine/Plugins/VirtualProduction/DMX/DMXPixelMapping/Source/DMXPixelMappingRuntime/Public/Components/DMXPixelMappingMatrixCellComponent.h
+++ b/Engine/Plugins/VirtualProduction/DMX/DMXPixelMapping/Source/DMXPixelMappingRuntime/Public/Components/DMXPixelMappingMatrixCellComponent.h
@@ -45,7 +45,6 @@ public:
 
 	//~ Begin UDMXPixelMappingBaseComponent implementation
 	virtual const FName& GetNamePrefix() override;
-	virtual void ResetDMX() override;
 	virtual FString GetUserName() const override;
 	//~ End UDMXPixelMappingBaseComponent implementation
 
diff --git a/Engine/Plugins/VirtualProduction/DMX/DMXPixelMapping/Source/DMXPixelMappingRuntime/Public/Components/DMXPixelMappingMatrixComponent.h b/Engine/Plugins/VirtualProduction/DMX/DMXPixelMapping/Source/DMXPixelMappingRuntime/Public/Components/DMXPixelMappingMatrixComponent.h
--- a/Engine/Plugins/VirtualProduction/DMX/DMXPixelMapping/Source/DMXPixelMappingRuntime/Public/Components/DMXPixelMappingMatrixComponent.h
+++ b/Engine/Plugins/VirtualProduction/DMX/DMXPixelMapping/Source/DMXPixelMappingRuntime/Public/Components/DMXPixelMappingMatrixComponent.h
@@ -58,7 +58,7 @@ public:
 public:
 	// ~Begin UDMXPixelMappingBaseComponent interface
 	virtual const FName& GetNamePrefix() override;
-	virtual void ResetDMX() override;
+	virtual void ResetDMX(EDMXPixelMappingResetDMXMode ResetMode = EDMXPixelMappingResetDMXMode::SendDefaultValues) override;
 	virtual void SendDMX() override;
 	virtual bool CanBeMovedTo(const UDMXPixelMappingBaseComponent* Component) const override;
 	virtual FString GetUserName() const override;
diff --git a/Engine/Plugins/VirtualProduction/DMX/DMXPixelMapping/Source/DMXPixelMappingRuntime/Public/Components/DMXPixelMappingRendererComponent.h b/Engine/Plugins/VirtualProduction/DMX/DMXPixelMapping/Source/DMXPixelMappingRuntime/Public/Components/DMXPixelMappingRendererComponent.h
--- a/Engine/Plugins/VirtualProduction/DMX/DMXPixelMapping/Source/DMXPixelMappingRuntime/Public/Components/DMXPixelMappingRendererComponent.h
+++ b/Engine/Plugins/VirtualProduction/DMX/DMXPixelMapping/Source/DMXPixelMappingRuntime/Public/Components/DMXPixelMappingRendererComponent.h
@@ -59,7 +59,7 @@ public:
 	//~ Begin UDMXPixelMappingBaseComponent implementation
 	virtual const FName& GetNamePrefix() override;
 	virtual bool CanBeMovedTo(const UDMXPixelMappingBaseComponent* Component) const override;
-	virtual void ResetDMX() override;
+	virtual void ResetDMX(EDMXPixelMappingResetDMXMode ResetMode = EDMXPixelMappingResetDMXMode::SendDefaultValues) override;
 	virtual void SendDMX() override;
 	virtual void Render() final;
 	virtual void RenderAndSendDMX() final;
diff --git a/Engine/Plugins/VirtualProduction/DMX/DMXPixelMapping/Source/DMXPixelMappingRuntime/Public/Components/DMXPixelMappingRootComponent.h b/Engine/Plugins/VirtualProduction/DMX/DMXPixelMapping/Source/DMXPixelMappingRuntime/Public/Components/DMXPixelMappingRootComponent.h
--- a/Engine/Plugins/VirtualProduction/DMX/DMXPixelMapping/Source/DMXPixelMappingRuntime/Public/Components/DMXPixelMappingRootComponent.h
+++ b/Engine/Plugins/VirtualProduction/DMX/DMXPixelMapping/Source/DMXPixelMappingRuntime/Public/Components/DMXPixelMappingRootComponent.h
@@ -29,7 +29,7 @@ public:
 	virtual const FName& GetNamePrefix() override;
 	virtual void AddChild(UDMXPixelMappingBaseComponent* InComponent) override;
 	virtual void RemoveChild(UDMXPixelMappingBaseComponent* InComponent) override;
-	virtual void ResetDMX() override;
+	virtual void ResetDMX(EDMXPixelMappingResetDMXMode ResetMode = EDMXPixelMappingResetDMXMode::SendDefaultValues) override;
 	virtual void SendDMX() override;
 	virtual void Render() override;
 	virtual void RenderAndSendDMX() override;
diff --git a/Engine/Plugins/VirtualProduction/DMX/DMXPixelMapping/Source/DMXPixelMappingRuntime/Public/Components/DMXPixelMappingScreenComponent.h b/Engine/Plugins/VirtualProduction/DMX/DMXPixelMapping/Source/DMXPixelMappingRuntime/Public/Components/DMXPixelMappingScreenComponent.h
--- a/Engine/Plugins/VirtualProduction/DMX/DMXPixelMapping/Source/DMXPixelMappingRuntime/Public/Components/DMXPixelMappingScreenComponent.h
+++ b/Engine/Plugins/VirtualProduction/DMX/DMXPixelMapping/Source/DMXPixelMappingRuntime/Public/Components/DMXPixelMappingScreenComponent.h
@@ -42,7 +42,6 @@ public:
 
 	//~ Begin UDMXPixelMappingBaseComponent implementation
 	virtual const FName& GetNamePrefix() override;
-	virtual void ResetDMX() override;
 	virtual void SendDMX() override;
 	//~ End UDMXPixelMappingBaseComponent implementation
 
