DMX - Pixel Mapping fix an issue where 'children follow size' sets the wrong size, does not properly transact

diff --git a/Engine/Plugins/VirtualProduction/DMX/DMXPixelMapping/Source/DMXPixelMappingRuntime/Private/Components/DMXPixelMappingRendererComponent.cpp b/Engine/Plugins/VirtualProduction/DMX/DMXPixelMapping/Source/DMXPixelMappingRuntime/Private/Components/DMXPixelMappingRendererComponent.cpp
--- a/Engine/Plugins/VirtualProduction/DMX/DMXPixelMapping/Source/DMXPixelMappingRuntime/Private/Components/DMXPixelMappingRendererComponent.cpp
+++ b/Engine/Plugins/VirtualProduction/DMX/DMXPixelMapping/Source/DMXPixelMappingRuntime/Private/Components/DMXPixelMappingRendererComponent.cpp
@@ -10,6 +10,7 @@
 #include "Components/DMXPixelMappingMatrixComponent.h"
 #include "Components/DMXPixelMappingRootComponent.h"
 #include "Components/DMXPixelMappingScreenComponent.h"
+#include "DMXPixelMapping.h"
 #include "DMXPixelMappingPixelMapRenderer.h"
 #include "DMXPixelMappingPreprocessRenderer.h"
 #include "DMXPixelMappingMainStreamObjectVersion.h"
@@ -337,7 +338,8 @@ void UDMXPixelMappingRendererComponent::OnComponentAddedOrRemoved(UDMXPixelMappi
 
 void UDMXPixelMappingRendererComponent::LetChildrenFollowSize()
 {
-	if (!bChildrenFollowSize || !PreprocessRenderer)
+	UDMXPixelMapping* PixelMapping = GetPixelMapping();
+	if (!PixelMapping || !bChildrenFollowSize || !PreprocessRenderer)
 	{
 		return;
 	}
@@ -359,6 +361,11 @@ void UDMXPixelMappingRendererComponent::LetChildrenFollowSize()
 		return;
 	}
 
+#if WITH_EDITOR
+	// In editor, temporarily disable scale children with parent
+	TGuardValue Guard(PixelMapping->bEditorScaleChildrenWithParent, false);
+#endif 
+
 	// Scale position and size of all children
 	const FVector2D Scalar = NewSize / LayoutRect;
 
@@ -368,6 +375,8 @@ void UDMXPixelMappingRendererComponent::LetChildrenFollowSize()
 		{
 			if (UDMXPixelMappingOutputComponent* OutputComponent = Cast<UDMXPixelMappingOutputComponent>(Component))
 			{
+				OutputComponent->Modify();
+
 				const FVector2D NewPosition = OutputComponent->GetPosition() * Scalar;
 				OutputComponent->SetPosition(NewPosition);
 
@@ -378,6 +387,7 @@ void UDMXPixelMappingRendererComponent::LetChildrenFollowSize()
 		bRecursive);
 
 	// Remember the new layout rect
+	Modify();
 	LayoutRect = NewSize;
 }
 
