DMX - Fix an issue with entangled references in DMXPixelMappingOutputComponentModel

diff --git a/Engine/Plugins/VirtualProduction/DMX/DMXPixelMapping/Source/DMXPixelMappingEditor/Private/Widgets/SDMXPixelMappingOutputComponent.cpp b/Engine/Plugins/VirtualProduction/DMX/DMXPixelMapping/Source/DMXPixelMappingEditor/Private/Widgets/SDMXPixelMappingOutputComponent.cpp
--- a/Engine/Plugins/VirtualProduction/DMX/DMXPixelMapping/Source/DMXPixelMappingEditor/Private/Widgets/SDMXPixelMappingOutputComponent.cpp
+++ b/Engine/Plugins/VirtualProduction/DMX/DMXPixelMapping/Source/DMXPixelMappingEditor/Private/Widgets/SDMXPixelMappingOutputComponent.cpp
@@ -33,7 +33,7 @@ namespace UE::DMX
 
 		ParentCanvas = InCanvas;
 
-		ParentCanvas->AddSlot()
+		InCanvas->AddSlot()
 			.ZOrder(0)
 			.AutoSize(true)
 			.Alignment(FVector2D::ZeroVector)
@@ -52,7 +52,7 @@ namespace UE::DMX
 	{
 		if (ParentCanvas.IsValid() && Slot)
 		{
-			ParentCanvas->RemoveSlot(Slot->GetWidget());
+			ParentCanvas.Pin()->RemoveSlot(Slot->GetWidget());
 		}
 
 		Slot = nullptr;
diff --git a/Engine/Plugins/VirtualProduction/DMX/DMXPixelMapping/Source/DMXPixelMappingEditor/Private/Widgets/SDMXPixelMappingOutputComponent.h b/Engine/Plugins/VirtualProduction/DMX/DMXPixelMapping/Source/DMXPixelMappingEditor/Private/Widgets/SDMXPixelMappingOutputComponent.h
--- a/Engine/Plugins/VirtualProduction/DMX/DMXPixelMapping/Source/DMXPixelMappingEditor/Private/Widgets/SDMXPixelMappingOutputComponent.h
+++ b/Engine/Plugins/VirtualProduction/DMX/DMXPixelMapping/Source/DMXPixelMappingEditor/Private/Widgets/SDMXPixelMappingOutputComponent.h
@@ -45,7 +45,7 @@ namespace UE::DMX
 
 	private:
 		/** When added to a parent, the canvas it was added to */
-		TSharedPtr<SConstraintCanvas> ParentCanvas;
+		TWeakPtr<SConstraintCanvas> ParentCanvas;
 
 		/** The canvas slot of the component widget */
 		SConstraintCanvas::FSlot* Slot = nullptr;
@@ -121,7 +121,7 @@ namespace UE::DMX
 		bool bDrawsSimplisticView = false;
 
 		/** When added to a parent, the canvas it was added to */
-		TSharedPtr<SConstraintCanvas> ParentCanvas;
+		TWeakPtr<SConstraintCanvas> ParentCanvas;
 
 		/** The canvas slot of the component widget */
 		SConstraintCanvas::FSlot* Slot = nullptr;
