DMX - PixelMapping now has an option to flip components of fixture groups horizontally or verically, new components find a reasonable size when first added

diff --git a/Engine/Plugins/VirtualProduction/DMX/DMXPixelMapping/Source/DMXPixelMappingEditor/Private/DMXPixelMappingEditorCommands.cpp b/Engine/Plugins/VirtualProduction/DMX/DMXPixelMapping/Source/DMXPixelMappingEditor/Private/DMXPixelMappingEditorCommands.cpp
--- a/Engine/Plugins/VirtualProduction/DMX/DMXPixelMapping/Source/DMXPixelMappingEditor/Private/DMXPixelMappingEditorCommands.cpp
+++ b/Engine/Plugins/VirtualProduction/DMX/DMXPixelMapping/Source/DMXPixelMappingEditor/Private/DMXPixelMappingEditorCommands.cpp
@@ -35,8 +35,10 @@ void FDMXPixelMappingEditorCommands::RegisterCommands()
 
 	UI_COMMAND(ToggleGridSnapping, "Toggle Grid Snapping", "Enables/disables grid snapping", EUserInterfaceActionType::Button, FInputChord());
 
-	UI_COMMAND(SizeComponentToTexture, "Size Component to Texture", "Sizes the selected group to Texture.", EUserInterfaceActionType::Button, FInputChord(EModifierKey::Shift, EKeys::S));
-	
+	UI_COMMAND(SizeGroupToTexture, "Size Group to Texture", "Sizes the selected group to Texture.", EUserInterfaceActionType::Button, FInputChord());
+	UI_COMMAND(FlipGroupHorizontally, "Flip Group Horizontally", "Flips children of selected group along its relative Y-Axis.", EUserInterfaceActionType::Button, FInputChord());
+	UI_COMMAND(FlipGroupVertically, "Flip Group Verticaclly", "Flips children of selected group along its relative X-Axis.", EUserInterfaceActionType::Button, FInputChord());
+
 	UI_COMMAND(ToggleScaleChildrenWithParent, "Scale Children with Parent", "Sets if children are scaled with their parent, when the parent is resized.", EUserInterfaceActionType::ToggleButton, FInputChord(EModifierKey::Alt, EKeys::Q));
 	UI_COMMAND(ToggleAlwaysSelectGroup, "Always select Group", "Sets if the parent is selected, if a child is clicked.", EUserInterfaceActionType::ToggleButton, FInputChord(EModifierKey::Alt, EKeys::W));
 	UI_COMMAND(ToggleApplyLayoutScriptWhenLoaded, "Apply Layout Script instantly", "Sets if layout script are applied as soon as they are loaded.", EUserInterfaceActionType::ToggleButton, FInputChord(EModifierKey::Alt, EKeys::E));
diff --git a/Engine/Plugins/VirtualProduction/DMX/DMXPixelMapping/Source/DMXPixelMappingEditor/Private/DMXPixelMappingEditorCommands.h b/Engine/Plugins/VirtualProduction/DMX/DMXPixelMapping/Source/DMXPixelMappingEditor/Private/DMXPixelMappingEditorCommands.h
--- a/Engine/Plugins/VirtualProduction/DMX/DMXPixelMapping/Source/DMXPixelMappingEditor/Private/DMXPixelMappingEditorCommands.h
+++ b/Engine/Plugins/VirtualProduction/DMX/DMXPixelMapping/Source/DMXPixelMappingEditor/Private/DMXPixelMappingEditorCommands.h
@@ -33,8 +33,10 @@ public:
 
 	TSharedPtr<FUICommandInfo> ToggleGridSnapping;
 
-	TSharedPtr<FUICommandInfo> SizeComponentToTexture;
-
+	TSharedPtr<FUICommandInfo> FlipGroupHorizontally;
+	TSharedPtr<FUICommandInfo> FlipGroupVertically;
+	TSharedPtr<FUICommandInfo> SizeGroupToTexture;
+	
 	TSharedPtr<FUICommandInfo> ToggleScaleChildrenWithParent;
 	TSharedPtr<FUICommandInfo> ToggleAlwaysSelectGroup;
 	TSharedPtr<FUICommandInfo> ToggleApplyLayoutScriptWhenLoaded;
diff --git a/Engine/Plugins/VirtualProduction/DMX/DMXPixelMapping/Source/DMXPixelMappingEditor/Private/Settings/DMXPixelMappingEditorSettings.h b/Engine/Plugins/VirtualProduction/DMX/DMXPixelMapping/Source/DMXPixelMappingEditor/Private/Settings/DMXPixelMappingEditorSettings.h
--- a/Engine/Plugins/VirtualProduction/DMX/DMXPixelMapping/Source/DMXPixelMappingEditor/Private/Settings/DMXPixelMappingEditorSettings.h
+++ b/Engine/Plugins/VirtualProduction/DMX/DMXPixelMapping/Source/DMXPixelMappingEditor/Private/Settings/DMXPixelMappingEditorSettings.h
@@ -15,7 +15,7 @@ struct FDMXPixelMappingHierarchySettings
 
 	/** If true, shows the editor color column */
 	UPROPERTY()
-	bool bShowEditorColorColumn = false;
+	bool bShowEditorColorColumn = true;
 
 	/** If true, shows the fixture ID column */
 	UPROPERTY()
diff --git a/Engine/Plugins/VirtualProduction/DMX/DMXPixelMapping/Source/DMXPixelMappingEditor/Private/Toolkits/DMXPixelMappingToolkit.cpp b/Engine/Plugins/VirtualProduction/DMX/DMXPixelMapping/Source/DMXPixelMappingEditor/Private/Toolkits/DMXPixelMappingToolkit.cpp
--- a/Engine/Plugins/VirtualProduction/DMX/DMXPixelMapping/Source/DMXPixelMappingEditor/Private/Toolkits/DMXPixelMappingToolkit.cpp
+++ b/Engine/Plugins/VirtualProduction/DMX/DMXPixelMapping/Source/DMXPixelMappingEditor/Private/Toolkits/DMXPixelMappingToolkit.cpp
@@ -4,6 +4,13 @@
 
 #include "Algo/Sort.h"
 #include "Algo/Transform.h"
+#include "Components/DMXPixelMappingFixtureGroupComponent.h"
+#include "Components/DMXPixelMappingFixtureGroupItemComponent.h"
+#include "Components/DMXPixelMappingRendererComponent.h"
+#include "Components/DMXPixelMappingRootComponent.h"
+#include "Components/DMXPixelMappingMatrixCellComponent.h"
+#include "Components/DMXPixelMappingMatrixComponent.h"
+#include "Components/DMXPixelMappingScreenComponent.h"
 #include "DMXPixelMapping.h"
 #include "DMXPixelMappingEditorCommands.h"
 #include "DMXPixelMappingEditorModule.h"
@@ -14,11 +21,8 @@
 #include "Framework/Commands/GenericCommands.h"
 #include "Framework/MultiBox/MultiBoxExtender.h"
 #include "K2Node_PixelMappingBaseComponent.h"
-#include "Components/DMXPixelMappingFixtureGroupComponent.h"
-#include "Components/DMXPixelMappingRendererComponent.h"
-#include "Components/DMXPixelMappingRootComponent.h"
-#include "Components/DMXPixelMappingMatrixComponent.h"
-#include "Components/DMXPixelMappingScreenComponent.h"
+#include "Library/DMXEntityFixturePatch.h"
+#include "Library/DMXLibrary.h"
 #include "Modules/ModuleManager.h"
 #include "ScopedTransaction.h"
 #include "Settings/DMXPixelMappingEditorSettings.h"
@@ -134,8 +138,8 @@ void FDMXPixelMappingToolkit::InitPixelMappingEditor(const EToolkitMode::Type Mo
 			)
 		);
 
-	const bool bCreateDefaultStandaloneMenu = true;
-	const bool bCreateDefaultToolbar = true;
+	constexpr bool bCreateDefaultStandaloneMenu = true;
+	constexpr bool bCreateDefaultToolbar = true;
 	FAssetEditorToolkit::InitAssetEditor(Mode, InitToolkitHost, FDMXPixelMappingEditorModule::DMXPixelMappingEditorAppIdentifier,
 		StandaloneDefaultLayout, bCreateDefaultStandaloneMenu, bCreateDefaultToolbar, InDMXPixelMapping);
 
@@ -146,20 +150,20 @@ void FDMXPixelMappingToolkit::InitPixelMappingEditor(const EToolkitMode::Type Mo
 	// Make an initial selection
 	if (UDMXPixelMappingRootComponent* RootComponent = InDMXPixelMapping->GetRootComponent())
 	{
-		UDMXPixelMappingBaseComponent* const* FirstRendererComponetPtr = Algo::FindByPredicate(InDMXPixelMapping->GetRootComponent()->GetChildren(), [](UDMXPixelMappingBaseComponent* Component)
+		UDMXPixelMappingBaseComponent* const* FirstRendererComponentPtr = Algo::FindByPredicate(InDMXPixelMapping->GetRootComponent()->GetChildren(), [](UDMXPixelMappingBaseComponent* Component)
 			{
 				return Component && Component->GetClass() == UDMXPixelMappingRendererComponent::StaticClass();
 			});
-		if (FirstRendererComponetPtr)
+		if (FirstRendererComponentPtr)
 		{
-			UDMXPixelMappingBaseComponent* const* FirstFixtureGroupComponetPtr = Algo::FindByPredicate((*FirstRendererComponetPtr)->GetChildren(), [](UDMXPixelMappingBaseComponent* Component)
+			UDMXPixelMappingBaseComponent* const* FirstFixtureGroupComponentPtr = Algo::FindByPredicate((*FirstRendererComponentPtr)->GetChildren(), [](UDMXPixelMappingBaseComponent* Component)
 				{
 					return Component && Component->GetClass() == UDMXPixelMappingFixtureGroupComponent::StaticClass();
 				});
 
-			if (UDMXPixelMappingBaseComponent* const* ComponentToSelectPtr = FirstFixtureGroupComponetPtr ?
-				FirstFixtureGroupComponetPtr :
-				FirstRendererComponetPtr)
+			if (UDMXPixelMappingBaseComponent* const* ComponentToSelectPtr = FirstFixtureGroupComponentPtr ?
+				FirstFixtureGroupComponentPtr :
+				FirstRendererComponentPtr)
 			{
 				const FDMXPixelMappingComponentReference ComponentReference(StaticCastSharedRef<FDMXPixelMappingToolkit>(AsShared()), *ComponentToSelectPtr);
 				SelectComponents(TSet<FDMXPixelMappingComponentReference>({ ComponentReference }));
@@ -497,9 +501,22 @@ TArray<UDMXPixelMappingBaseComponent*> FDMXPixelMappingToolkit::CreateComponents
 
 					Target->AddChild(NewComponent);
 
-					// Output components need to adopt the initial rotation from their parent if possible.
+					// Find a reasonable size if components are added to a fixture group
 					UDMXPixelMappingOutputComponent* NewOutputComponent = Cast<UDMXPixelMappingOutputComponent>(NewComponent);
 					UDMXPixelMappingOutputComponent* ParentOutputComponent = NewOutputComponent ? Cast<UDMXPixelMappingOutputComponent>(NewOutputComponent->GetParent()) : nullptr;
+					UDMXPixelMappingFixtureGroupComponent* GroupComponent = Target ? Cast<UDMXPixelMappingFixtureGroupComponent>(ParentOutputComponent) : nullptr;
+					UDMXLibrary* DMXLibrary = GroupComponent->DMXLibrary;
+					if (GroupComponent && DMXLibrary)
+					{
+						const TArray<UDMXEntityFixturePatch*> FixturePatches = DMXLibrary->GetEntitiesTypeCast<UDMXEntityFixturePatch>();
+						const int32 Columns = FMath::RoundFromZero(FMath::Sqrt((float)FixturePatches.Num()));
+						const int32 Rows = FMath::RoundFromZero((float)FixturePatches.Num() / Columns);
+						const FVector2D Size = FVector2D(GroupComponent->GetSize().X / Columns, GroupComponent->GetSize().Y / Rows);
+
+						NewOutputComponent->SetSize(Size);
+					}
+
+					// Output components need to adopt the initial rotation from their parent if possible
 					if (NewOutputComponent && ParentOutputComponent)
 					{
 						NewOutputComponent->SetRotation(ParentOutputComponent->GetRotation());
@@ -562,74 +579,95 @@ void FDMXPixelMappingToolkit::DeleteSelectedComponents()
 	UpdateBlueprintNodes();
 }	
 
-bool FDMXPixelMappingToolkit::CanSizeSelectedComponentToTexture() const
+bool FDMXPixelMappingToolkit::CanPerformCommandsOnGroup() const
 {
-	if (SelectedComponents.Num() == 1)
-	{
-		if (UDMXPixelMappingBaseComponent* Component = SelectedComponents.Array()[0].GetComponent())
-		{
-			return
-				Component->GetClass() == UDMXPixelMappingFixtureGroupComponent::StaticClass() ||
-				Component->GetClass() == UDMXPixelMappingMatrixComponent::StaticClass() ||
-				Component->GetClass() == UDMXPixelMappingScreenComponent::StaticClass();
-		}
-	}
-	return false;
+	const UDMXPixelMappingFixtureGroupComponent* FixtureGroupComponent = GetFixtureGroupFromSelection(); 
+	return FixtureGroupComponent != nullptr;
 }
 
-void FDMXPixelMappingToolkit::SizeSelectedComponentToTexture(bool bTransacted)
+void FDMXPixelMappingToolkit::FlipGroup(EOrientation Orientation, bool bTransacted)
 {
-	if (!ensureMsgf(CanSizeSelectedComponentToTexture(), TEXT("Trying to size selected component to texture without previously testing CanSizeSelectedComponentToTexture.")))
+	if (!ensureMsgf(CanPerformCommandsOnGroup(), TEXT("Trying to flip cells without previously testing CanPerformCommandsOnGroup.")))
 	{
 		return;
 	}
 
-	if (SelectedComponents.IsEmpty())
+	UDMXPixelMappingFixtureGroupComponent* FixtureGroupComponent = GetFixtureGroupFromSelection();
+	if (!FixtureGroupComponent)
 	{
 		return;
 	}
 
-	const UDMXPixelMappingRendererComponent* RendererComponent = GetActiveRendererComponent();
-	if (!RendererComponent)
+	TSharedPtr<FScopedTransaction> FilpCellsTransaction;
+	if (bTransacted)
+	{
+		const FText TransactionText = FText::Format(LOCTEXT("FilpCellsTransaction", "Flip Group {0}"), Orientation == EOrientation::Orient_Horizontal ? 
+			LOCTEXT("FlipHorizontalText", "Horizontally") :
+			LOCTEXT("FlipHorizontalText", "Vertically"));
+
+		FilpCellsTransaction = MakeShared<FScopedTransaction>(TransactionText);
+	}
+
+	const double RestoreRotation = FixtureGroupComponent->GetRotation();
+	FixtureGroupComponent->SetRotation(0.0);
+
+	const FVector2D Center = FixtureGroupComponent->GetPosition() + FixtureGroupComponent->GetSize() / 2.f;
+
+	constexpr bool bRecursive = false;
+	FixtureGroupComponent->ForEachChildOfClass<UDMXPixelMappingOutputComponent>([&Center, Orientation](UDMXPixelMappingOutputComponent* Child)
+		{
+			const FVector2D ChildPivotOffset = Child->GetSize() / 2.f;
+			const FVector2D ChildCenter = Child->GetPosition() + Child->GetSize() / 2.f;
+			const FVector2D NewPositionBothAxes = Center + Center - ChildCenter - ChildPivotOffset;
+			if (Orientation == EOrientation::Orient_Horizontal)
+			{
+				Child->SetPosition(FVector2D(NewPositionBothAxes.X, Child->GetPosition().Y));
+			}
+			else
+			{
+				Child->SetPosition(FVector2D(Child->GetPosition().X, NewPositionBothAxes.Y));
+			}
+		},
+		bRecursive);
+
+	FixtureGroupComponent->SetRotation(RestoreRotation);
+}
+
+void FDMXPixelMappingToolkit::SizeGroupToTexture(bool bTransacted)
+{
+	if (!ensureMsgf(CanPerformCommandsOnGroup(), TEXT("Trying to size selected component to texture without previously testing CanPerformCommandsOnGroup.")))
 	{
 		return;
 	}
 
-	const FVector2D TextureSize = RendererComponent->GetSize();
-	if (TextureSize == FVector2D::ZeroVector)
+	UDMXPixelMappingFixtureGroupComponent* FixtureGroupComponent = GetFixtureGroupFromSelection();
+	if (!FixtureGroupComponent)
 	{
 		return;
 	}
 
-	UDMXPixelMappingOutputComponent* Component = Cast<UDMXPixelMappingOutputComponent>(SelectedComponents.Array()[0].GetComponent());
-	if (!Component)
+	const UDMXPixelMappingRendererComponent* RendererComponent = GetActiveRendererComponent();
+	if (!RendererComponent)
 	{
 		return;
 	}
 
-	// Apply to parent where appropriate
-	if (Component->GetClass() != UDMXPixelMappingFixtureGroupComponent::StaticClass() &&
-		Component->GetClass() != UDMXPixelMappingMatrixComponent::StaticClass() &&
-		Component->GetClass() != UDMXPixelMappingScreenComponent::StaticClass())
+	const FVector2D TextureSize = RendererComponent->GetSize();
+	if (TextureSize == FVector2D::ZeroVector)
 	{
-		if ((Component->GetParent() && Component->GetParent()->GetClass() == UDMXPixelMappingFixtureGroupComponent::StaticClass()) ||
-			(Component->GetParent() && Component->GetParent()->GetClass() == UDMXPixelMappingMatrixComponent::StaticClass()) ||
-			(Component->GetParent() && Component->GetParent()->GetClass() == UDMXPixelMappingScreenComponent::StaticClass()))
-		{
-			Component = Cast<UDMXPixelMappingOutputComponent>(Component->GetParent());
-		}
+		return;
 	}
 
-	TSharedPtr<FScopedTransaction> SizeComponentToTextureTransaction;
+	TSharedPtr<FScopedTransaction> SizeGroupToTextureTransaction;
 	if (bTransacted)
 	{ 
-		SizeComponentToTextureTransaction = MakeShared<FScopedTransaction>(LOCTEXT("SizeComponentToTextureTransaction", "Size Component to Texture"));
+		SizeGroupToTextureTransaction = MakeShared<FScopedTransaction>(LOCTEXT("SizeGroupToTextureTransaction", "Size Group to Texture"));
 	}
 
-	Component->Modify();
-	Component->SetRotation(0.0);
-	Component->SetPosition(FVector2D::ZeroVector);
-	Component->SetSize(TextureSize);
+	FixtureGroupComponent->Modify();
+	FixtureGroupComponent->SetRotation(0.0);
+	FixtureGroupComponent->SetPosition(FVector2D::ZeroVector);
+	FixtureGroupComponent->SetSize(TextureSize);
 }
 
 void FDMXPixelMappingToolkit::SetTransformHandleMode(EDMXPixelMappingTransformHandleMode NewTransformHandleMode)
@@ -686,7 +724,6 @@ void FDMXPixelMappingToolkit::OnComponentRenamed(UDMXPixelMappingBaseComponent*
 	UpdateBlueprintNodes();
 }
 
-
 TSharedRef<SDockTab> FDMXPixelMappingToolkit::SpawnTab_DMXLibraryView(const FSpawnTabArgs& Args)
 {
 	check(Args.GetTabId() == DMXLibraryViewTabID);
@@ -994,11 +1031,29 @@ void FDMXPixelMappingToolkit::SetupCommands()
 	);
 
 	// Designer related
-	constexpr bool bSizeComponentToTexture = true;
+	constexpr bool bTransact = true;
 	GetToolkitCommands()->MapAction(
-		FDMXPixelMappingEditorCommands::Get().SizeComponentToTexture,
-		FExecuteAction::CreateSP(this, &FDMXPixelMappingToolkit::SizeSelectedComponentToTexture, bSizeComponentToTexture),
-		FCanExecuteAction::CreateSP(this, &FDMXPixelMappingToolkit::CanSizeSelectedComponentToTexture)
+		FDMXPixelMappingEditorCommands::Get().FlipGroupHorizontally,
+		FExecuteAction::CreateSP(this, &FDMXPixelMappingToolkit::FlipGroup, EOrientation::Orient_Horizontal, bTransact),
+		FCanExecuteAction::CreateSP(this, &FDMXPixelMappingToolkit::CanPerformCommandsOnGroup),
+		FIsActionChecked(),
+		FIsActionButtonVisible::CreateSP(this, &FDMXPixelMappingToolkit::CanPerformCommandsOnGroup)
+	);
+	
+	GetToolkitCommands()->MapAction(
+		FDMXPixelMappingEditorCommands::Get().FlipGroupVertically,
+		FExecuteAction::CreateSP(this, &FDMXPixelMappingToolkit::FlipGroup, EOrientation::Orient_Vertical, bTransact),
+		FCanExecuteAction::CreateSP(this, &FDMXPixelMappingToolkit::CanPerformCommandsOnGroup),
+		FIsActionChecked(),
+		FIsActionButtonVisible::CreateSP(this, &FDMXPixelMappingToolkit::CanPerformCommandsOnGroup)
+	);
+
+	GetToolkitCommands()->MapAction(
+		FDMXPixelMappingEditorCommands::Get().SizeGroupToTexture,
+		FExecuteAction::CreateSP(this, &FDMXPixelMappingToolkit::SizeGroupToTexture, bTransact),
+		FCanExecuteAction::CreateSP(this, &FDMXPixelMappingToolkit::CanPerformCommandsOnGroup),
+		FIsActionChecked(),
+		FIsActionButtonVisible::CreateSP(this, &FDMXPixelMappingToolkit::CanPerformCommandsOnGroup)
 	);
 
 	UDMXPixelMappingEditorSettings* EditorSettings = GetMutableDefault<UDMXPixelMappingEditorSettings>(); 
@@ -1036,4 +1091,48 @@ ECheckBoxState FDMXPixelMappingToolkit::GetTransformHandleModeCheckboxState(EDMX
 	return CompareTransformHandleMode == TransformHandleMode ? ECheckBoxState::Checked : ECheckBoxState::Unchecked;
 }
 
+UDMXPixelMappingFixtureGroupComponent* FDMXPixelMappingToolkit::GetFixtureGroupFromSelection() const
+{
+	TArray<UDMXPixelMappingFixtureGroupComponent*> FixtureGroupComponents;
+	for (const FDMXPixelMappingComponentReference& ComponentReference : SelectedComponents)
+	{
+		UDMXPixelMappingBaseComponent* Component = ComponentReference.GetComponent();
+		if (!Component)
+		{
+			continue;
+		}
+
+		UDMXPixelMappingFixtureGroupComponent* FixtureGroupComponent = nullptr;
+		if (Component->GetClass() == UDMXPixelMappingFixtureGroupComponent::StaticClass())
+		{
+			FixtureGroupComponent = Cast<UDMXPixelMappingFixtureGroupComponent>(Component);
+		}
+		if ((Component->GetClass() == UDMXPixelMappingFixtureGroupItemComponent::StaticClass() ||
+			Component->GetClass() == UDMXPixelMappingMatrixComponent::StaticClass()) &&
+			Component->GetParent())
+		{
+			FixtureGroupComponent = Cast<UDMXPixelMappingFixtureGroupComponent>(Component->GetParent());
+		}
+		else if (Component->GetClass() == UDMXPixelMappingMatrixCellComponent::StaticClass() &&
+			Component->GetParent() &&
+			Component->GetParent()->GetParent())
+		{
+			FixtureGroupComponent = Cast<UDMXPixelMappingFixtureGroupComponent>(Component->GetParent()->GetParent());
+		}
+
+		if (FixtureGroupComponent)
+		{
+			FixtureGroupComponents.AddUnique(FixtureGroupComponent);
+		}
+	}
+
+	// Return the group only if exactly one is contained in selection
+	if (FixtureGroupComponents.Num() == 1)
+	{
+		return FixtureGroupComponents[0];
+	}
+
+	return nullptr;
+}
+
 #undef LOCTEXT_NAMESPACE
diff --git a/Engine/Plugins/VirtualProduction/DMX/DMXPixelMapping/Source/DMXPixelMappingEditor/Private/Toolkits/DMXPixelMappingToolkit.h b/Engine/Plugins/VirtualProduction/DMX/DMXPixelMapping/Source/DMXPixelMappingEditor/Private/Toolkits/DMXPixelMappingToolkit.h
--- a/Engine/Plugins/VirtualProduction/DMX/DMXPixelMapping/Source/DMXPixelMappingEditor/Private/Toolkits/DMXPixelMappingToolkit.h
+++ b/Engine/Plugins/VirtualProduction/DMX/DMXPixelMapping/Source/DMXPixelMappingEditor/Private/Toolkits/DMXPixelMappingToolkit.h
@@ -25,6 +25,7 @@ class SDMXPixelMappingLayoutView;
 class SDMXPixelMappingPreviewView;
 class UDMXPixelMapping;
 class UDMXPixelMappingBaseComponent;
+class UDMXPixelMappingFixtureGroupComponent;
 class UDMXPixelMappingMatrixComponent;
 class UDMXPixelMappingOutputComponent;
 class UDMXPixelMappingRendererComponent;
@@ -146,11 +147,14 @@ public:
 	/** Deletes the selected Components */
 	void DeleteSelectedComponents();
 
-	/** Returns true if the selected component can be sized to texture */
-	bool CanSizeSelectedComponentToTexture() const;
+	/** Returns true if the current selection yields a fixture group on which commands can be performed on, e.g. SizeGroupToTexture */
+	bool CanPerformCommandsOnGroup() const;
+
+	/** Flips children of a group either horizontally or vertically */
+	void FlipGroup(EOrientation Orientation, bool bTransacted);
 
 	/** Sizes the component to the render target of the pixelmapping asset */
-	void SizeSelectedComponentToTexture(bool bTransacted);
+	void SizeGroupToTexture(bool bTransacted);
 
 	/** Toggles between grid snapping enabled and disabled */
 	void ToggleGridSnapping();
@@ -227,6 +231,12 @@ private:
 	/** Returns the checkbox state for a transform handle mode, checked if the mode equals the current mode. */
 	ECheckBoxState GetTransformHandleModeCheckboxState(EDMXPixelMappingTransformHandleMode CompareTransformHandleMode) const;
 
+	/** 
+	 * Returns the fixture group from a selection. Returns the parent group if the primary selection is a child of a group.
+	 * Returns the primary fixture group in selection or nullptr if no fixture group can be deduced from selection. 
+	 */
+	UDMXPixelMappingFixtureGroupComponent* GetFixtureGroupFromSelection() const;
+
 	/** List of open tool panels; used to ensure only one exists at any one time */
 	TMap<FName, TWeakPtr<SDockableTab>> SpawnedToolPanels;
 
diff --git a/Engine/Plugins/VirtualProduction/DMX/DMXPixelMapping/Source/DMXPixelMappingEditor/Private/ViewModels/DMXPixelMappingDMXLibraryViewModel.cpp b/Engine/Plugins/VirtualProduction/DMX/DMXPixelMapping/Source/DMXPixelMappingEditor/Private/ViewModels/DMXPixelMappingDMXLibraryViewModel.cpp
--- a/Engine/Plugins/VirtualProduction/DMX/DMXPixelMapping/Source/DMXPixelMappingEditor/Private/ViewModels/DMXPixelMappingDMXLibraryViewModel.cpp
+++ b/Engine/Plugins/VirtualProduction/DMX/DMXPixelMapping/Source/DMXPixelMappingEditor/Private/ViewModels/DMXPixelMappingDMXLibraryViewModel.cpp
@@ -57,10 +57,10 @@ void UDMXPixelMappingDMXLibraryViewModel::CreateAndSetNewFixtureGroup(TWeakPtr<F
 	if (OtherFixtureGroupComponents.IsEmpty())
 	{
 		// If there's no group, add one that scales the texture of the active renderer component
-		if (Toolkit->CanSizeSelectedComponentToTexture())
+		if (Toolkit->CanPerformCommandsOnGroup())
 		{
 			constexpr bool bTransacted = false;
-			Toolkit->SizeSelectedComponentToTexture(bTransacted);
+			Toolkit->SizeGroupToTexture(bTransacted);
 		}
 	}
 	else
@@ -343,10 +343,16 @@ void UDMXPixelMappingDMXLibraryViewModel::LayoutEvenOverParent(const TArray<UDMX
 
 void UDMXPixelMappingDMXLibraryViewModel::LayoutAfterLastPatch(const TArray<UDMXPixelMappingBaseComponent*> Components)
 {
-	if (Components.IsEmpty())
+	if (Components.IsEmpty() || !DMXLibrary)
 	{
 		return;
 	}
+	const TArray<UDMXEntityFixturePatch*> Patches = DMXLibrary->GetEntitiesTypeCast<UDMXEntityFixturePatch>();
+	if (Patches.IsEmpty())
+	{
+		return;
+	}
+
 	UDMXPixelMappingOutputComponent* FirstComponentToLayout = Cast<UDMXPixelMappingOutputComponent>(Components[0]);
 
 	UDMXPixelMappingFixtureGroupComponent* FixtureGroupComponent = WeakFixtureGroupComponent.Get();
@@ -374,8 +380,8 @@ void UDMXPixelMappingDMXLibraryViewModel::LayoutAfterLastPatch(const TArray<UDMX
 	{
 		if (UDMXPixelMappingOutputComponent* OtherOutputComponent = Cast<UDMXPixelMappingOutputComponent>(OtherComponent))
 		{
-			NextPosition.X = FMath::Max(NextPosition.X, OtherOutputComponent->GetPosition().X + OtherOutputComponent->GetSize().X + 1);
-			if (NextPosition.X + FirstComponentToLayout->GetSize().X > FixtureGroupComponent->GetPosition().X + FixtureGroupComponent->GetSize().X - 1)
+			NextPosition.X = FMath::Max(NextPosition.X, OtherOutputComponent->GetPosition().X + OtherOutputComponent->GetSize().X);
+			if (NextPosition.X + FirstComponentToLayout->GetSize().X > FixtureGroupComponent->GetPosition().X + FixtureGroupComponent->GetSize().X)
 			{
 				NextPosition.X = FixtureGroupComponent->GetPosition().X;
 				NextPosition.Y += FirstComponentToLayout->GetSize().Y;
@@ -389,11 +395,10 @@ void UDMXPixelMappingDMXLibraryViewModel::LayoutAfterLastPatch(const TArray<UDMX
 	{
 		if (UDMXPixelMappingOutputComponent* OutputComponent = Cast<UDMXPixelMappingOutputComponent>(Component))
 		{
-			if (FixtureGroupComponent->IsOverPosition(NextPosition) &&
-				FixtureGroupComponent->IsOverPosition(NextPosition + OutputComponent->GetSize()))
-			{
-				OutputComponent->SetPosition(NextPosition);
+			OutputComponent->SetPosition(NextPosition);
 
+			if (OutputComponent->IsOverParent())
+			{
 				RowHeight = FMath::Max(OutputComponent->GetSize().Y, RowHeight);
 				NextPosition = FVector2D(NextPosition.X + OutputComponent->GetSize().X, NextPosition.Y);
 			}
@@ -401,18 +406,18 @@ void UDMXPixelMappingDMXLibraryViewModel::LayoutAfterLastPatch(const TArray<UDMX
 			{
 				// Try on a new row
 				FVector2D NewRowPosition = FVector2D(FixtureGroupComponent->GetPosition().X, NextPosition.Y + RowHeight);
-
 				const FVector2D NextPositionOnNewRow = FVector2D(NewRowPosition.X + OutputComponent->GetSize().X, NewRowPosition.Y);
-				if (FixtureGroupComponent->IsOverPosition(NextPositionOnNewRow) &&
-					FixtureGroupComponent->IsOverPosition(NextPositionOnNewRow + OutputComponent->GetSize()))
-				{
-					OutputComponent->SetPosition(NewRowPosition);
 
+				OutputComponent->SetPosition(NextPositionOnNewRow);
+
+				if (OutputComponent->IsOverParent())
+				{
 					NextPosition = FVector2D(NewRowPosition.X + OutputComponent->GetSize().X, NewRowPosition.Y);
 					RowHeight = OutputComponent->GetSize().Y;
 				}
 				else
 				{
+					// Append as the component cannot be fit into the group
 					OutputComponent->SetPosition(NextPosition);
 
 					NextPosition = FVector2D(NextPosition.X + OutputComponent->GetSize().X, NextPosition.Y);
diff --git a/Engine/Plugins/VirtualProduction/DMX/DMXPixelMapping/Source/DMXPixelMappingEditor/Private/Views/SDMXPixelMappingDesignerView.cpp b/Engine/Plugins/VirtualProduction/DMX/DMXPixelMapping/Source/DMXPixelMappingEditor/Private/Views/SDMXPixelMappingDesignerView.cpp
--- a/Engine/Plugins/VirtualProduction/DMX/DMXPixelMapping/Source/DMXPixelMappingEditor/Private/Views/SDMXPixelMappingDesignerView.cpp
+++ b/Engine/Plugins/VirtualProduction/DMX/DMXPixelMapping/Source/DMXPixelMappingEditor/Private/Views/SDMXPixelMappingDesignerView.cpp
@@ -784,11 +784,13 @@ void SDMXPixelMappingDesignerView::OpenContextMenu()
 	}
 	MenuBuilder.EndSection();
 
-	if (Toolkit->CanSizeSelectedComponentToTexture())
+	if (Toolkit->CanPerformCommandsOnGroup())
 	{
-		MenuBuilder.BeginSection("AdvancedActions", LOCTEXT("ActionsSection", "Advanced Actions"));
+		MenuBuilder.BeginSection("FixtureGroupActions", LOCTEXT("FixtureGroupActionsSection", "Fixture Group Actions"));
 		{
-			MenuBuilder.AddMenuEntry(FDMXPixelMappingEditorCommands::Get().SizeComponentToTexture);
+			MenuBuilder.AddMenuEntry(FDMXPixelMappingEditorCommands::Get().FlipGroupHorizontally);
+			MenuBuilder.AddMenuEntry(FDMXPixelMappingEditorCommands::Get().FlipGroupVertically);
+			MenuBuilder.AddMenuEntry(FDMXPixelMappingEditorCommands::Get().SizeGroupToTexture);
 		}
 		MenuBuilder.EndSection();
 	}
