DMX - Rename DMXPixelMappingDynamicRange and related property to PixelFormat to align with similar type names in engine

diff --git a/Engine/Plugins/VirtualProduction/DMX/DMXPixelMapping/Source/DMXPixelMappingRuntime/Private/Components/DMXPixelMappingRendererComponent.cpp b/Engine/Plugins/VirtualProduction/DMX/DMXPixelMapping/Source/DMXPixelMappingRuntime/Private/Components/DMXPixelMappingRendererComponent.cpp
--- a/Engine/Plugins/VirtualProduction/DMX/DMXPixelMapping/Source/DMXPixelMappingRuntime/Private/Components/DMXPixelMappingRendererComponent.cpp
+++ b/Engine/Plugins/VirtualProduction/DMX/DMXPixelMapping/Source/DMXPixelMappingRuntime/Private/Components/DMXPixelMappingRendererComponent.cpp
@@ -125,7 +125,7 @@ void UDMXPixelMappingRendererComponent::PostEditChangeChainProperty(FPropertyCha
 	}	
 	
 	const FName PropertyName = PropertyChangedChainEvent.GetPropertyName();
-	if (PropertyName == GET_MEMBER_NAME_CHECKED(UDMXPixelMappingRendererComponent, DynamicRange))
+	if (PropertyName == GET_MEMBER_NAME_CHECKED(UDMXPixelMappingRendererComponent, PixelFormat))
 	{
 		InvalidatePixelMapRenderer();
 	}
@@ -395,30 +395,30 @@ EPixelFormat UDMXPixelMappingRendererComponent::GetFormatFromDynamicRange() cons
 {
 	UTexture2D* InputTexture2D = Cast<UTexture2D>(InputTexture);
 
-	if (DynamicRange == EDMXPixelMappingRendererDynamicRange::Auto &&
+	if (PixelFormat == EDMXPixelMappingRendererPixelFormat::Auto &&
 		RendererType == EDMXPixelMappingRendererType::Texture &&
 		InputTexture2D)
 	{
-		const EPixelFormat PixelFormat = InputTexture2D->GetPixelFormat();
+		const EPixelFormat TexturePixelFormat = InputTexture2D->GetPixelFormat();
 	
 		// Propagonate the pixel format, if it is supported by render targets. 
 		// As there doesn't seem to be an engine call to do this conversion,
 		// use (the opposite) logic of GetPixelFormatFromRenderTargetFormat in TextureRenderTarget2D.h (5.4).
-		if (PixelFormat == PF_G8 ||
-			PixelFormat == PF_R8G8 ||
-			PixelFormat == PF_B8G8R8A8 ||
-			PixelFormat == PF_R16F ||
-			PixelFormat == PF_G16R16F ||
-			PixelFormat == PF_FloatRGBA ||
-			PixelFormat == PF_R32_FLOAT ||
-			PixelFormat == PF_G32R32F ||
-			PixelFormat == PF_A32B32G32R32F ||
-			PixelFormat == PF_A2B10G10R10)
+		if (TexturePixelFormat == PF_G8 ||
+			TexturePixelFormat == PF_R8G8 ||
+			TexturePixelFormat == PF_B8G8R8A8 ||
+			TexturePixelFormat == PF_R16F ||
+			TexturePixelFormat == PF_G16R16F ||
+			TexturePixelFormat == PF_FloatRGBA ||
+			TexturePixelFormat == PF_R32_FLOAT ||
+			TexturePixelFormat == PF_G32R32F ||
+			TexturePixelFormat == PF_A32B32G32R32F ||
+			TexturePixelFormat == PF_A2B10G10R10)
 		{
-			return PixelFormat;
+			return TexturePixelFormat;
 		}
 	}
-	else if (DynamicRange == EDMXPixelMappingRendererDynamicRange::RGBA16F)
+	else if (PixelFormat == EDMXPixelMappingRendererPixelFormat::RGBA16F)
 	{
 		return PF_FloatRGBA;
 	}
diff --git a/Engine/Plugins/VirtualProduction/DMX/DMXPixelMapping/Source/DMXPixelMappingRuntime/Public/Components/DMXPixelMappingRendererComponent.h b/Engine/Plugins/VirtualProduction/DMX/DMXPixelMapping/Source/DMXPixelMappingRuntime/Public/Components/DMXPixelMappingRendererComponent.h
--- a/Engine/Plugins/VirtualProduction/DMX/DMXPixelMapping/Source/DMXPixelMappingRuntime/Public/Components/DMXPixelMappingRendererComponent.h
+++ b/Engine/Plugins/VirtualProduction/DMX/DMXPixelMapping/Source/DMXPixelMappingRuntime/Public/Components/DMXPixelMappingRendererComponent.h
@@ -23,7 +23,7 @@ namespace UE::DMXPixelMapping::Rendering { class FPixelMapRenderElement; }
 
 
 UENUM(BlueprintType)
-enum class EDMXPixelMappingRendererDynamicRange : uint8
+enum class EDMXPixelMappingRendererPixelFormat : uint8
 {
 	Auto UMETA(ToolTip = "Use the format of the Input Texture, use 8-bit for Input Materials and UMG."),
 	RGBA8 UMETA(DisplayName = "Low Precision 8-bit (RGBA8)"),
@@ -109,7 +109,7 @@ public:
 
 	/** The dynamic range of the renderer */
 	UPROPERTY(EditAnywhere, BlueprintReadOnly, Category = "Render Settings")
-	EDMXPixelMappingRendererDynamicRange DynamicRange = EDMXPixelMappingRendererDynamicRange::Auto;
+	EDMXPixelMappingRendererPixelFormat PixelFormat = EDMXPixelMappingRendererPixelFormat::Auto;
 
 	/** The exposure of the pixel mapping renderer. This property affects the DMX output. */
 	UPROPERTY(EditAnywhere, BlueprintReadWrite, Category = "Render Settings", Meta = (ClampMin = "0", UIMin = "0", UIMax = "1", DisplayName = "Exposure"))
