DMX - Fix CIS 'GroupComponent' pointer was utilized before it was verified against nullptr

diff --git a/Engine/Plugins/VirtualProduction/DMX/DMXPixelMapping/Source/DMXPixelMappingEditor/Private/Toolkits/DMXPixelMappingToolkit.cpp b/Engine/Plugins/VirtualProduction/DMX/DMXPixelMapping/Source/DMXPixelMappingEditor/Private/Toolkits/DMXPixelMappingToolkit.cpp
--- a/Engine/Plugins/VirtualProduction/DMX/DMXPixelMapping/Source/DMXPixelMappingEditor/Private/Toolkits/DMXPixelMappingToolkit.cpp
+++ b/Engine/Plugins/VirtualProduction/DMX/DMXPixelMapping/Source/DMXPixelMappingEditor/Private/Toolkits/DMXPixelMappingToolkit.cpp
@@ -501,12 +501,12 @@ TArray<UDMXPixelMappingBaseComponent*> FDMXPixelMappingToolkit::CreateComponents
 
 					Target->AddChild(NewComponent);
 
-					// Find a reasonable size if components are added to a fixture group
+					// Find a reasonable size when components are added to a fixture group
 					UDMXPixelMappingOutputComponent* NewOutputComponent = Cast<UDMXPixelMappingOutputComponent>(NewComponent);
-					UDMXPixelMappingOutputComponent* ParentOutputComponent = NewOutputComponent ? Cast<UDMXPixelMappingOutputComponent>(NewOutputComponent->GetParent()) : nullptr;
-					UDMXPixelMappingFixtureGroupComponent* GroupComponent = Target ? Cast<UDMXPixelMappingFixtureGroupComponent>(ParentOutputComponent) : nullptr;
-					UDMXLibrary* DMXLibrary = GroupComponent->DMXLibrary;
-					if (GroupComponent && DMXLibrary)
+					const UDMXPixelMappingOutputComponent* ParentOutputComponent = NewOutputComponent ? Cast<UDMXPixelMappingOutputComponent>(NewOutputComponent->GetParent()) : nullptr;
+					const UDMXPixelMappingFixtureGroupComponent* GroupComponent = ParentOutputComponent ? Cast<UDMXPixelMappingFixtureGroupComponent>(ParentOutputComponent) : nullptr;
+					const UDMXLibrary* DMXLibrary = GroupComponent ? GroupComponent->DMXLibrary : nullptr;
+					if (NewOutputComponent && GroupComponent && DMXLibrary)
 					{
 						const TArray<UDMXEntityFixturePatch*> FixturePatches = DMXLibrary->GetEntitiesTypeCast<UDMXEntityFixturePatch>();
 						const int32 Columns = FMath::RoundFromZero(FMath::Sqrt((float)FixturePatches.Num()));
