DMX - Rename files and classes that were not using a DMXPixelMapping specific prefix

diff --git a/Engine/Plugins/VirtualProduction/DMX/DMXPixelMapping/Source/DMXPixelMappingRenderer/Private/DMXPixelMappingApplyFilterMaterialProxy.cpp b/Engine/Plugins/VirtualProduction/DMX/DMXPixelMapping/Source/DMXPixelMappingRenderer/Private/DMXPixelMappingApplyFilterMaterialProxy.cpp
new file mode 100644
--- /dev/null
+++ b/Engine/Plugins/VirtualProduction/DMX/DMXPixelMapping/Source/DMXPixelMappingRenderer/Private/DMXPixelMappingApplyFilterMaterialProxy.cpp
@@ -0,0 +1,292 @@
+// Copyright Epic Games, Inc. All Rights Reserved.
+
+#include "DMXPixelMappingApplyFilterMaterialProxy.h"
+
+#include "CanvasTypes.h"
+#include "CommonRenderResources.h"
+#include "DMXStats.h"
+#include "Engine/Canvas.h"
+#include "Engine/Engine.h"
+#include "Engine/Texture.h"
+#include "Engine/TextureRenderTarget2D.h"
+#include "GlobalShader.h"
+#include "Materials/MaterialInstanceDynamic.h"
+#include "Materials/MaterialInterface.h"
+#include "Materials/MaterialRenderProxy.h"
+#include "Modules/ModuleManager.h"
+#include "RHIStaticStates.h"
+#include "ScreenRendering.h"
+#include "TextureResource.h"
+
+
+namespace UE::DMXPixelMapping::Rendering::Preprocess::Private
+{
+	DECLARE_CYCLE_STAT(TEXT("PixelMapping ApplyFilterMaterial"), STAT_PreprocessApplyFilterMaterial, STATGROUP_DMX);
+
+	FDMXPixelMappingPreprocessMaterialCanvas::FDMXPixelMappingPreprocessMaterialCanvas()
+	{
+		Canvas = NewObject<UCanvas>();
+	}
+
+	void FDMXPixelMappingPreprocessMaterialCanvas::DrawMaterialToRenderTarget(UTextureRenderTarget2D* TextureRenderTarget, UMaterialInterface* Material)
+	{
+		if (!Material || !TextureRenderTarget || !TextureRenderTarget->GetResource())
+		{
+			return;
+		}
+
+		// This is a user-facing function, so we'd rather make sure that shaders are ready by the time we render, in order to ensure we don't draw with a fallback material
+		Material->EnsureIsComplete();
+		FTextureRenderTargetResource* RenderTargetResource = TextureRenderTarget->GameThread_GetRenderTargetResource();
+		if (!RenderTargetResource)
+		{
+			return;
+		}
+
+		FCanvas RenderCanvas(
+			RenderTargetResource,
+			nullptr,
+			GEngine->GetWorld(),
+			GMaxRHIFeatureLevel);
+
+		Canvas->Init(TextureRenderTarget->SizeX, TextureRenderTarget->SizeY, nullptr, &RenderCanvas);
+		{
+			ENQUEUE_RENDER_COMMAND(FlushDeferredResourceUpdateCommand)(
+				[RenderTargetResource](FRHICommandListImmediate& RHICmdList)
+				{
+					RenderTargetResource->FlushDeferredResourceUpdate(RHICmdList);
+				});
+
+			Canvas->K2_DrawMaterial(Material, FVector2D(0, 0), FVector2D(TextureRenderTarget->SizeX, TextureRenderTarget->SizeY), FVector2D(0, 0));
+
+			RenderCanvas.Flush_GameThread();
+			Canvas->Canvas = nullptr;
+
+			// UpdateResourceImmediate must be called here to ensure mips are generated.
+			TextureRenderTarget->UpdateResourceImmediate(false);
+
+			ENQUEUE_RENDER_COMMAND(ResetSceneTextureExtentHistory)(
+				[RenderTargetResource](FRHICommandListImmediate& RHICmdList)
+				{
+					RenderTargetResource->ResetSceneTextureExtentsHistory();
+				});
+		}
+	}
+
+	void FDMXPixelMappingPreprocessMaterialCanvas::AddReferencedObjects(FReferenceCollector& Collector)
+	{
+		Collector.AddReferencedObject(Canvas);
+	}
+
+	FDMXPixelMappingApplyFilterMaterialProxy::FDMXPixelMappingApplyFilterMaterialProxy(EPixelFormat InFormat)
+		: Format(InFormat)
+	{}
+
+	void FDMXPixelMappingApplyFilterMaterialProxy::Render(UTexture* InInputTexture, const UDMXPixelMappingPreprocessRenderer& InPreprocessRenderer)
+	{
+		SCOPE_CYCLE_COUNTER(STAT_PreprocessApplyFilterMaterial);
+
+		MaterialInstanceDynamic = InPreprocessRenderer.GetFilterMID();
+
+		WeakInputTexture = InInputTexture;
+		UpdateRenderTargets(InPreprocessRenderer.GetNumDownsamplePasses(), InPreprocessRenderer.GetDesiredOutputSize2D());
+
+		if (!CanRender())
+		{
+			return;
+		}
+
+		UTexture* InputTexture = WeakInputTexture.Get();
+		check(InputTexture);
+		if (MaterialInstanceDynamic && DownsampleRenderTargets.IsEmpty() && ensureMsgf(OutputRenderTarget, TEXT("Cannot apply pixel mapping filter material. Missing output render target.")))
+		{
+			// Render without downsampling
+			MaterialInstanceDynamic->SetTextureParameterValue(InPreprocessRenderer.GetTextureParameterName(), InputTexture);
+			MaterialInstanceDynamic->SetScalarParameterValue(InPreprocessRenderer.GetBlurDistanceParameterName(), InPreprocessRenderer.GetBlurDistance());
+			DrawMaterialCanvas.DrawMaterialToRenderTarget(OutputRenderTarget, MaterialInstanceDynamic);
+		}
+		else
+		{
+			const bool bApplyFilterMaterialEachDownsamplePass = InPreprocessRenderer.ShouldApplyFilterMaterialEachDownsamplePass();
+
+			// Render with downsampling. 
+			// Note, if the last downsample render target is of minimal size, further renderer targets are not created.
+			// Hence use the number of downsample render targets instead of InParams.NumDownsamplePasses.
+			const int32 NumDownsamplePasses = DownsampleRenderTargets.Num();
+			for (int32 DownsamplePass = 0; DownsamplePass < NumDownsamplePasses; DownsamplePass++)
+			{
+				UTexture* SourceTexture = DownsamplePass == 0 ? InputTexture : DownsampleRenderTargets[DownsamplePass - 1];
+				UTextureRenderTarget2D* DownsampleRenderTarget = DownsampleRenderTargets[DownsamplePass].Get();
+
+				const bool bApplyFilterMaterial = MaterialInstanceDynamic && (bApplyFilterMaterialEachDownsamplePass || DownsampleRenderTarget == DownsampleRenderTargets.Last());
+				if (bApplyFilterMaterial)
+				{
+					MaterialInstanceDynamic->SetTextureParameterValue(InPreprocessRenderer.GetTextureParameterName(), SourceTexture);
+					MaterialInstanceDynamic->SetScalarParameterValue(InPreprocessRenderer.GetBlurDistanceParameterName(), InPreprocessRenderer.GetBlurDistance());
+					DrawMaterialCanvas.DrawMaterialToRenderTarget(DownsampleRenderTarget, MaterialInstanceDynamic);
+				}
+				else
+				{
+					RenderTextureToTarget(SourceTexture, DownsampleRenderTarget);
+				}
+			}
+
+			// Scale to output size
+			UTexture* LastSourceTexture = DownsampleRenderTargets.IsEmpty() ? InputTexture : DownsampleRenderTargets.Last();
+			RenderTextureToTarget(LastSourceTexture, OutputRenderTarget);
+		}
+	}
+
+	UTexture* FDMXPixelMappingApplyFilterMaterialProxy::GetRenderedTexture() const
+	{
+		return CanRender() ? OutputRenderTarget : WeakInputTexture.Get();
+	}
+
+	void FDMXPixelMappingApplyFilterMaterialProxy::AddReferencedObjects(FReferenceCollector& Collector)
+	{
+		Collector.AddReferencedObject(MaterialInstanceDynamic);
+		Collector.AddReferencedObjects(DownsampleRenderTargets);
+		Collector.AddReferencedObject(OutputRenderTarget);
+	}
+
+	bool FDMXPixelMappingApplyFilterMaterialProxy::CanRender() const
+	{
+		return WeakInputTexture.IsValid() && OutputRenderTarget && FApp::CanEverRender();
+	}
+
+	void FDMXPixelMappingApplyFilterMaterialProxy::UpdateRenderTargets(int32 NumDownsamplePasses, const TOptional<FVector2D>& OptionalOutputSize)
+	{
+		if (!WeakInputTexture.IsValid())
+		{
+			OutputRenderTarget = nullptr;
+			return;
+		}
+		UTexture* InputTexture = WeakInputTexture.Get();
+		check(InputTexture);
+		
+		const bool bResize = 
+			OptionalOutputSize.IsSet() &&
+			InputTexture->GetSurfaceWidth() != OptionalOutputSize.GetValue().X &&
+			InputTexture->GetSurfaceHeight() != OptionalOutputSize.GetValue().Y;
+
+		const bool bNeedsAnyRenderTargets = NumDownsamplePasses > 0 || MaterialInstanceDynamic || bResize;
+		if (!bNeedsAnyRenderTargets)
+		{
+			OutputRenderTarget = nullptr;
+			return;
+		}
+
+		const int32 NumExpectedDownsampleRenderTargets = OptionalOutputSize.IsSet() ? NumDownsamplePasses : NumDownsamplePasses - 1;
+		const bool bInvalidRenderTargets = 
+			!OutputRenderTarget || 
+			NumExpectedDownsampleRenderTargets != DownsampleRenderTargets.Num() ||
+			(OptionalOutputSize.IsSet() && OutputRenderTarget->GetSurfaceWidth() != OptionalOutputSize.GetValue().X) ||
+			(OptionalOutputSize.IsSet() && OutputRenderTarget->GetSurfaceHeight() != OptionalOutputSize.GetValue().Y);
+
+		if (bInvalidRenderTargets)
+		{
+			FlushRenderingCommands();
+			
+			const FVector2D InputSize = FVector2D(InputTexture->GetSurfaceWidth(), InputTexture->GetSurfaceHeight());
+
+			// Create downsample render targets
+			DownsampleRenderTargets.Reset();
+			FVector2D DownsampleSize = InputSize;
+			for (int32 DownsamplePass = 0; DownsamplePass < NumDownsamplePasses; DownsamplePass++)
+			{
+				DownsampleSize /= 2.0;
+				if (DownsampleSize.X <= 1.0 || DownsampleSize.Y <= 1.0)
+				{
+					break;
+				}
+
+				constexpr bool bForceLinearGamma = false;
+
+				UTextureRenderTarget2D* DownsampleRenderTarget = NewObject<UTextureRenderTarget2D>();
+				DownsampleRenderTarget->ClearColor = FLinearColor::Black;
+				DownsampleRenderTarget->InitCustomFormat(DownsampleSize.X, DownsampleSize.Y, Format, bForceLinearGamma);
+				DownsampleRenderTarget->UpdateResourceImmediate();
+				DownsampleRenderTargets.Add(DownsampleRenderTarget);
+			}
+
+			UTextureRenderTarget2D* ScaleRenderTarget = nullptr;
+			if (OptionalOutputSize.IsSet())
+			{
+				constexpr bool bForceLinearGamma = false;
+
+				ScaleRenderTarget = NewObject<UTextureRenderTarget2D>();
+				ScaleRenderTarget->ClearColor = FLinearColor::Black;
+				ScaleRenderTarget->InitCustomFormat(OptionalOutputSize.GetValue().X, OptionalOutputSize.GetValue().Y, Format, bForceLinearGamma);
+				ScaleRenderTarget->UpdateResourceImmediate();
+			}
+
+			if (ScaleRenderTarget)
+			{
+				OutputRenderTarget = ScaleRenderTarget;
+			}
+			else if (!DownsampleRenderTargets.IsEmpty())
+			{
+				OutputRenderTarget = DownsampleRenderTargets.Pop();
+			}
+			else
+			{
+				OutputRenderTarget = nullptr;
+			}
+		}
+	}
+
+	void FDMXPixelMappingApplyFilterMaterialProxy::RenderTextureToTarget(UTexture* Texture, UTextureRenderTarget2D* RenderTarget) const
+	{
+		IRendererModule* RendererModule = FModuleManager::GetModulePtr<IRendererModule>("Renderer");
+		ENQUEUE_RENDER_COMMAND(PixelMappingPreprocessDownsamplePass)(
+			[RendererModule, Texture, RenderTarget, this](FRHICommandListImmediate& RHICmdList)
+			{
+				if (!Texture || !Texture->GetResource())
+				{
+					return;
+				}
+
+				const FTextureRHIRef SourceTextureRHI = Texture->GetResource()->GetTexture2DRHI();
+				const FTextureRHIRef TargetTextureRHI = RenderTarget->GetResource()->GetTexture2DRHI();
+
+				FRHIRenderPassInfo RPInfo(TargetTextureRHI, MakeRenderTargetActions(ERenderTargetLoadAction::ENoAction, ERenderTargetStoreAction::EStore));
+
+				RHICmdList.Transition(FRHITransitionInfo(TargetTextureRHI, ERHIAccess::Unknown, ERHIAccess::RTV));
+				RHICmdList.BeginRenderPass(RPInfo, TEXT("PixelMappingFilterProxy_RenderToTarget"));
+				{
+					RHICmdList.SetViewport(0, 0, 0.0f, (float)TargetTextureRHI->GetSizeX(), (float)TargetTextureRHI->GetSizeY(), 1.0f);
+
+					FGraphicsPipelineStateInitializer GraphicsPSOInit;
+					RHICmdList.ApplyCachedRenderTargets(GraphicsPSOInit);
+					GraphicsPSOInit.BlendState = TStaticBlendState<>::GetRHI();
+					GraphicsPSOInit.RasterizerState = TStaticRasterizerState<>::GetRHI();
+					GraphicsPSOInit.DepthStencilState = TStaticDepthStencilState<false, CF_Always>::GetRHI();
+
+					FGlobalShaderMap* ShaderMap = GetGlobalShaderMap(GMaxRHIFeatureLevel);
+					TShaderMapRef<FScreenVS> VertexShader(ShaderMap);
+					TShaderMapRef<FScreenPS> PixelShader(ShaderMap);
+
+					GraphicsPSOInit.BoundShaderState.VertexDeclarationRHI = GFilterVertexDeclaration.VertexDeclarationRHI;
+					GraphicsPSOInit.BoundShaderState.VertexShaderRHI = VertexShader.GetVertexShader();
+					GraphicsPSOInit.BoundShaderState.PixelShaderRHI = PixelShader.GetPixelShader();
+					GraphicsPSOInit.PrimitiveType = PT_TriangleList;
+
+					SetGraphicsPipelineState(RHICmdList, GraphicsPSOInit, 0);
+
+					SetShaderParametersLegacyPS(RHICmdList, PixelShader, TStaticSamplerState<SF_Bilinear>::GetRHI(), SourceTextureRHI);
+
+					RendererModule->DrawRectangle(RHICmdList,
+						0, 0,									// Dest X, Y
+						(float)TargetTextureRHI->GetSizeX(),	// Dest Width
+						(float)TargetTextureRHI->GetSizeY(),	// Dest Height
+						0, 0,									// Source U, V
+						1, 1,									// Source USize, VSize
+						TargetTextureRHI->GetSizeXY(),			// Target buffer size
+						FIntPoint(1, 1),						// Source texture size
+						VertexShader,
+						EDRF_Default);
+				}
+				RHICmdList.EndRenderPass();
+			});
+	}
+}
diff --git a/Engine/Plugins/VirtualProduction/DMX/DMXPixelMapping/Source/DMXPixelMappingRenderer/Private/DMXPixelMappingApplyFilterMaterialProxy.h b/Engine/Plugins/VirtualProduction/DMX/DMXPixelMapping/Source/DMXPixelMappingRenderer/Private/DMXPixelMappingApplyFilterMaterialProxy.h
new file mode 100644
--- /dev/null
+++ b/Engine/Plugins/VirtualProduction/DMX/DMXPixelMapping/Source/DMXPixelMappingRenderer/Private/DMXPixelMappingApplyFilterMaterialProxy.h
@@ -0,0 +1,94 @@
+// Copyright Epic Games, Inc. All Rights Reserved.
+
+#pragma once
+
+#include "DMXPixelMappingPreprocessRenderer.h"
+
+#include "UObject/GCObject.h"
+#include "UObject/ObjectPtr.h"
+
+class UCanvas;
+class UMaterialInstanceDynamic;
+class UMaterialInterface;
+class UTexture;
+class UTextureRenderTarget2D;
+
+
+namespace UE::DMXPixelMapping::Rendering::Preprocess::Private
+{
+	/** Canvas implementation for pixelmapping */
+	class FDMXPixelMappingPreprocessMaterialCanvas
+		: public FGCObject
+	{
+	public:
+		/** Constructor */
+		FDMXPixelMappingPreprocessMaterialCanvas();
+
+		/** Draws a material to a render target */
+		void DrawMaterialToRenderTarget(UTextureRenderTarget2D* TextureRenderTarget, UMaterialInterface* Material);
+
+	protected:
+		//~ Begin FGCObject interface
+		virtual void AddReferencedObjects(FReferenceCollector& Collector) override;
+		virtual FString GetReferencerName() const override
+		{
+			return TEXT("DMXPixelMapping::Rendering::FDMXPixelMappingPreprocessMaterialCanvas");
+		}
+		//~ End FGCObject interface
+
+	private:
+		/** Canvas used to draw the material */
+		TObjectPtr<UCanvas> Canvas;
+	};
+
+	/** Proxy to filter the texture used in Pixel Mapping */
+	class FDMXPixelMappingApplyFilterMaterialProxy
+		: public IPreprocessApplyFilterMaterialProxy
+		, public FGCObject
+	{
+	public:
+		FDMXPixelMappingApplyFilterMaterialProxy(EPixelFormat InFormat);
+
+		//~ Begin IPreprocessApplyFilterMaterialProxy interface
+		virtual void Render(UTexture* InInputTexture, const UDMXPixelMappingPreprocessRenderer& InPreprocessRenderer) override;
+		virtual UTexture* GetRenderedTexture() const override;
+		//~ End IPreprocessApplyFilterMaterialProxy interface
+
+	protected:
+		//~ Begin FGCObject interface
+		virtual void AddReferencedObjects(FReferenceCollector& Collector) override;
+		virtual FString GetReferencerName() const override
+		{
+			return TEXT("FDMXPixelMappingApplyFilterMaterialProxy");
+		}
+		//~ End FGCObject interface
+
+	private:
+		/** Returns true if rendering is needed */
+		bool CanRender() const;
+
+		/** Updates render targets */
+		void UpdateRenderTargets(int32 NumDownsamplePasses, const TOptional<FVector2D>& OptionalOutputSize);
+
+		/** Renders the Input Texture to the Output render target */
+		void RenderTextureToTarget(UTexture* Texture, UTextureRenderTarget2D* RenderTarget) const;
+
+		/** The pixel format his proxy should use */
+		EPixelFormat Format = PF_Unknown;
+
+		/** Weak ref to the input texture */
+		TWeakObjectPtr<UTexture> WeakInputTexture;
+
+		/** The dynamic material instance to be applied to the input texture */
+		TObjectPtr<UMaterialInstanceDynamic> MaterialInstanceDynamic;
+
+		/** Downsample render targets */
+		TArray<TObjectPtr<UTextureRenderTarget2D>> DownsampleRenderTargets;
+
+		/** The target to which the output is rendered. Either the last downscale render target or the scale render target */
+		TObjectPtr<UTextureRenderTarget2D> OutputRenderTarget;
+
+		/** Canvas used to draw materials */
+		FDMXPixelMappingPreprocessMaterialCanvas DrawMaterialCanvas;
+	};
+}
diff --git a/Engine/Plugins/VirtualProduction/DMX/DMXPixelMapping/Source/DMXPixelMappingRenderer/Private/DMXPixelMappingPixelMapRenderer.cpp b/Engine/Plugins/VirtualProduction/DMX/DMXPixelMapping/Source/DMXPixelMappingRenderer/Private/DMXPixelMappingPixelMapRenderer.cpp
--- a/Engine/Plugins/VirtualProduction/DMX/DMXPixelMapping/Source/DMXPixelMappingRenderer/Private/DMXPixelMappingPixelMapRenderer.cpp
+++ b/Engine/Plugins/VirtualProduction/DMX/DMXPixelMapping/Source/DMXPixelMappingRenderer/Private/DMXPixelMappingPixelMapRenderer.cpp
@@ -3,8 +3,8 @@
 #include "DMXPixelMappingPixelMapRenderer.h"
 
 #include "DMXPixelMappingRenderElement.h"
+#include "DMXPixelMappingRenderPixelMapProxy.h"
 #include "Engine/TextureRenderTarget2D.h"
-#include "RenderPixelMapProxy.h"
 #include "RenderingThread.h"
 
 
@@ -13,7 +13,7 @@ const FIntPoint UDMXPixelMappingPixelMapRenderer::MaxPixelMapSize = FIntPoint(40
 UDMXPixelMappingPixelMapRenderer::UDMXPixelMappingPixelMapRenderer()
 {
 	using namespace UE::DMXPixelMapping::Rendering::Private;
-	RendererPixelMapProxy = MakeShared<FRenderPixelMapProxy>();
+	RendererPixelMapProxy = MakeShared<FDMXPixelMappingRenderPixelMapProxy>();
 }
 
 void UDMXPixelMappingPixelMapRenderer::SetElements(const TArray<TSharedRef<UE::DMXPixelMapping::Rendering::FPixelMapRenderElement>>& InElements, EPixelFormat InFormat)
diff --git a/Engine/Plugins/VirtualProduction/DMX/DMXPixelMapping/Source/DMXPixelMappingRenderer/Private/DMXPixelMappingPreprocessRenderer.cpp b/Engine/Plugins/VirtualProduction/DMX/DMXPixelMapping/Source/DMXPixelMappingRenderer/Private/DMXPixelMappingPreprocessRenderer.cpp
--- a/Engine/Plugins/VirtualProduction/DMX/DMXPixelMapping/Source/DMXPixelMappingRenderer/Private/DMXPixelMappingPreprocessRenderer.cpp
+++ b/Engine/Plugins/VirtualProduction/DMX/DMXPixelMapping/Source/DMXPixelMappingRenderer/Private/DMXPixelMappingPreprocessRenderer.cpp
@@ -2,13 +2,13 @@
 
 #include "DMXPixelMappingPreprocessRenderer.h"
 
+#include "DMXPixelMappingApplyFilterMaterialProxy.h"
+#include "DMXPixelMappingRenderInputMaterialProxy.h"
+#include "DMXPixelMappingRenderInputTextureProxy.h"
+#include "DMXPixelMappingRenderInputUserWidgetProxy.h"
 #include "DMXStats.h"
 #include "Engine/TextureRenderTarget2D.h"
 #include "Materials/MaterialInstanceDynamic.h"
-#include "PreprocessApplyFilterMaterialProxy.h"
-#include "PreprocessRenderInputMaterialProxy.h"
-#include "PreprocessRenderInputTextureProxy.h"
-#include "PreprocessRenderInputUserWidgetProxy.h"
 #include "UObject/Package.h"
 
 
@@ -17,8 +17,8 @@ DECLARE_CYCLE_STAT(TEXT("PixelMapping PreprocessInputTexture"), STAT_DMXPixelMap
 void UDMXPixelMappingPreprocessRenderer::SetInputTexture(UTexture* InTexture, EPixelFormat InFormat)
 {
 	using namespace UE::DMXPixelMapping::Rendering::Preprocess::Private;
-	RenderInputProxy = MakeShared<FPreprocessRenderInputTextureProxy>(InTexture);
-	ApplyMaterialProxy = MakeShared<FPreprocessApplyFilterMaterialProxy>(InFormat);
+	RenderInputProxy = MakeShared<FDMXPixelMappingRenderInputTextureProxy>(InTexture);
+	ApplyMaterialProxy = MakeShared<FDMXPixelMappingApplyFilterMaterialProxy>(InFormat);
 
 	bShowInputSize = false;
 }
@@ -26,8 +26,8 @@ void UDMXPixelMappingPreprocessRenderer::SetInputTexture(UTexture* InTexture, EP
 void UDMXPixelMappingPreprocessRenderer::SetInputMaterial(UMaterialInterface* InMaterial, EPixelFormat InFormat)
 {
 	using namespace UE::DMXPixelMapping::Rendering::Preprocess::Private;
-	RenderInputProxy = MakeShared<FPreprocessRenderInputMaterialProxy>(InMaterial, InputSize, InFormat);
-	ApplyMaterialProxy = MakeShared<FPreprocessApplyFilterMaterialProxy>(InFormat);
+	RenderInputProxy = MakeShared<FDMXPixelMappingRenderInputMaterialProxy>(InMaterial, InputSize, InFormat);
+	ApplyMaterialProxy = MakeShared<FDMXPixelMappingApplyFilterMaterialProxy>(InFormat);
 
 	bShowInputSize = true;
 }
@@ -35,8 +35,8 @@ void UDMXPixelMappingPreprocessRenderer::SetInputMaterial(UMaterialInterface* In
 void UDMXPixelMappingPreprocessRenderer::SetInputUserWidget(UUserWidget* InUserWidget, EPixelFormat InFormat)
 {
 	using namespace UE::DMXPixelMapping::Rendering::Preprocess::Private;
-	RenderInputProxy = MakeShared<FPreprocessRenderInputUserWidgetProxy>(InUserWidget, InputSize, InFormat);
-	ApplyMaterialProxy = MakeShared<FPreprocessApplyFilterMaterialProxy>(InFormat);
+	RenderInputProxy = MakeShared<FDMXPixelMappingRenderInputUserWidgetProxy>(InUserWidget, InputSize, InFormat);
+	ApplyMaterialProxy = MakeShared<FDMXPixelMappingApplyFilterMaterialProxy>(InFormat);
 
 	bShowInputSize = true;
 }
diff --git a/Engine/Plugins/VirtualProduction/DMX/DMXPixelMapping/Source/DMXPixelMappingRenderer/Private/DMXPixelMappingRenderInputMaterialProxy.cpp b/Engine/Plugins/VirtualProduction/DMX/DMXPixelMapping/Source/DMXPixelMappingRenderer/Private/DMXPixelMappingRenderInputMaterialProxy.cpp
new file mode 100644
--- /dev/null
+++ b/Engine/Plugins/VirtualProduction/DMX/DMXPixelMapping/Source/DMXPixelMappingRenderer/Private/DMXPixelMappingRenderInputMaterialProxy.cpp
@@ -0,0 +1,84 @@
+// Copyright Epic Games, Inc. All Rights Reserved.
+
+#include "DMXPixelMappingRenderInputMaterialProxy.h"
+
+#include "DMXPixelMappingPreprocessRenderer.h"
+#include "DMXStats.h"
+#include "Engine/TextureRenderTarget2D.h"
+#include "Materials/Material.h"
+#include "Materials/MaterialInterface.h"
+#include "SlateMaterialBrush.h"
+#include "Slate/WidgetRenderer.h"
+#include "Widgets/Images/SImage.h"
+
+
+namespace UE::DMXPixelMapping::Rendering::Preprocess::Private
+{
+	DECLARE_CYCLE_STAT(TEXT("PixelMapping RenderInputMaterial"), STAT_DMXPixelMappingPreprocessRenderInputMaterial, STATGROUP_DMX);
+
+	FDMXPixelMappingRenderInputMaterialProxy::FDMXPixelMappingRenderInputMaterialProxy(UMaterialInterface* InMaterial, const FVector2D& InInputSize, EPixelFormat InFormat)
+		: WeakMaterial(InMaterial)
+	{
+		if (InMaterial)
+		{
+			InMaterial->EnsureIsComplete();
+		}
+
+		constexpr bool bForceLinearGamma = false;
+
+		IntermediateRenderTarget = NewObject<UTextureRenderTarget2D>();
+		IntermediateRenderTarget->ClearColor = FLinearColor::Black;
+		IntermediateRenderTarget->InitCustomFormat(InInputSize.X, InInputSize.Y, InFormat, bForceLinearGamma);
+		IntermediateRenderTarget->UpdateResourceImmediate();
+
+		UIMaterialBrush = MakeShared<FSlateMaterialBrush>(FVector2D(1.f));
+
+		const bool bUseGammaCorrection = false;
+		MaterialRenderer = MakeShared<FWidgetRenderer>(bUseGammaCorrection);
+	}
+
+	void FDMXPixelMappingRenderInputMaterialProxy::Render()
+	{
+		SCOPE_CYCLE_COUNTER(STAT_DMXPixelMappingPreprocessRenderInputMaterial);
+
+		UMaterial* Material = WeakMaterial.IsValid() ? WeakMaterial->GetMaterial() : nullptr;
+		if (Material && Material->IsUIMaterial())
+		{
+			check(IntermediateRenderTarget);
+			FVector2D TextureSize = FVector2D(IntermediateRenderTarget->SizeX, IntermediateRenderTarget->SizeY);
+
+			UIMaterialBrush->ImageSize = TextureSize;
+			UIMaterialBrush->SetMaterial(Material);
+
+			TSharedRef<SWidget> Widget =
+				SNew(SOverlay)
+				+ SOverlay::Slot()
+				[
+					SNew(SImage)
+					.Image(UIMaterialBrush.Get())
+				];
+
+			static const float DeltaTime = 0.f; 
+			MaterialRenderer->DrawWidget(IntermediateRenderTarget, Widget, TextureSize, DeltaTime);
+
+			// Reset material after drawing
+			UIMaterialBrush->SetMaterial(nullptr);
+		}
+	}
+
+	UTexture* FDMXPixelMappingRenderInputMaterialProxy::GetRenderedTexture() const
+	{
+		return IntermediateRenderTarget;
+	}
+
+	FVector2D FDMXPixelMappingRenderInputMaterialProxy::GetSize2D() const
+	{
+		check(IntermediateRenderTarget);
+		return FVector2D(IntermediateRenderTarget->GetSurfaceWidth(), IntermediateRenderTarget->GetSurfaceHeight());
+	}
+
+	void FDMXPixelMappingRenderInputMaterialProxy::AddReferencedObjects(FReferenceCollector& Collector)
+	{
+		Collector.AddReferencedObject(IntermediateRenderTarget);
+	}
+}
diff --git a/Engine/Plugins/VirtualProduction/DMX/DMXPixelMapping/Source/DMXPixelMappingRenderer/Private/DMXPixelMappingRenderInputMaterialProxy.h b/Engine/Plugins/VirtualProduction/DMX/DMXPixelMapping/Source/DMXPixelMappingRenderer/Private/DMXPixelMappingRenderInputMaterialProxy.h
new file mode 100644
--- /dev/null
+++ b/Engine/Plugins/VirtualProduction/DMX/DMXPixelMapping/Source/DMXPixelMappingRenderer/Private/DMXPixelMappingRenderInputMaterialProxy.h
@@ -0,0 +1,53 @@
+// Copyright Epic Games, Inc. All Rights Reserved.
+
+#pragma once
+
+#include "DMXPixelMappingPreprocessRenderer.h"
+
+#include "UObject/GCObject.h"
+
+struct FSlateMaterialBrush;
+class FWidgetRenderer;
+class UMaterialInterface;
+class UTextureRenderTarget2D;
+
+
+namespace UE::DMXPixelMapping::Rendering::Preprocess::Private
+{
+	/** Proxy responsible for the input material */
+	class FDMXPixelMappingRenderInputMaterialProxy
+		: public IPreprocessRenderInputProxy
+		, public FGCObject
+	{
+	public:
+		FDMXPixelMappingRenderInputMaterialProxy(UMaterialInterface* InMaterial, const FVector2D& InInputSize, EPixelFormat Format);
+		
+		//~ Begin IPreprocessRenderInputProxy interface
+		virtual void Render() override;
+		virtual UTexture* GetRenderedTexture() const override;
+		virtual FVector2D GetSize2D() const override;
+		//~ End IPreprocessRenderInputProxy interface
+
+	protected:
+		//~ Begin FGCObject interface
+		virtual void AddReferencedObjects(FReferenceCollector& Collector) override;
+		virtual FString GetReferencerName() const override
+		{
+			return TEXT("DMXPixelMapping::Rendering::Private::FDMXPixelMappingRenderInputMaterialProxy");
+		}
+		//~ End FGCObject interface
+
+	private:
+		/** Material renderer */
+		TSharedPtr<FWidgetRenderer> MaterialRenderer;
+
+		/** Brush for Material widget renderer */
+		TSharedPtr<FSlateMaterialBrush> UIMaterialBrush;
+
+		/** The Input object */
+		TWeakObjectPtr<UMaterialInterface> WeakMaterial;
+
+		/** Intermediate render target for materials and user widgets */
+		TObjectPtr<UTextureRenderTarget2D> IntermediateRenderTarget;
+	};
+}
diff --git a/Engine/Plugins/VirtualProduction/DMX/DMXPixelMapping/Source/DMXPixelMappingRenderer/Private/DMXPixelMappingRenderInputTextureProxy.cpp b/Engine/Plugins/VirtualProduction/DMX/DMXPixelMapping/Source/DMXPixelMappingRenderer/Private/DMXPixelMappingRenderInputTextureProxy.cpp
new file mode 100644
--- /dev/null
+++ b/Engine/Plugins/VirtualProduction/DMX/DMXPixelMapping/Source/DMXPixelMappingRenderer/Private/DMXPixelMappingRenderInputTextureProxy.cpp
@@ -0,0 +1,35 @@
+// Copyright Epic Games, Inc. All Rights Reserved.
+
+#include "DMXPixelMappingRenderInputTextureProxy.h"
+
+#include "Engine/Texture.h"
+
+
+namespace UE::DMXPixelMapping::Rendering::Preprocess::Private
+{
+	FDMXPixelMappingRenderInputTextureProxy::FDMXPixelMappingRenderInputTextureProxy(UTexture* InTexture)
+		:  WeakInputTexture(InTexture)
+	{
+		if (InTexture)
+		{
+			InTexture->WaitForStreaming();
+		}
+	}
+
+	void FDMXPixelMappingRenderInputTextureProxy::Render()
+	{
+		// No rendering required
+	}
+
+	UTexture* FDMXPixelMappingRenderInputTextureProxy::GetRenderedTexture() const
+	{
+		return WeakInputTexture.Get();
+	}
+
+	FVector2D FDMXPixelMappingRenderInputTextureProxy::GetSize2D() const
+	{
+		return WeakInputTexture.IsValid() ?
+			FVector2D(WeakInputTexture->GetSurfaceWidth(), WeakInputTexture->GetSurfaceHeight()) :
+			FVector2D{ 0.0, 0.0 };
+	}
+}
diff --git a/Engine/Plugins/VirtualProduction/DMX/DMXPixelMapping/Source/DMXPixelMappingRenderer/Private/DMXPixelMappingRenderInputTextureProxy.h b/Engine/Plugins/VirtualProduction/DMX/DMXPixelMapping/Source/DMXPixelMappingRenderer/Private/DMXPixelMappingRenderInputTextureProxy.h
new file mode 100644
--- /dev/null
+++ b/Engine/Plugins/VirtualProduction/DMX/DMXPixelMapping/Source/DMXPixelMappingRenderer/Private/DMXPixelMappingRenderInputTextureProxy.h
@@ -0,0 +1,32 @@
+// Copyright Epic Games, Inc. All Rights Reserved.
+
+#pragma once
+
+#include "DMXPixelMappingPreprocessRenderer.h"
+
+#include "UObject/WeakObjectPtr.h"
+
+class UDMXPixelMappingPreprocessRenderer;
+class UTexture;
+
+
+namespace UE::DMXPixelMapping::Rendering::Preprocess::Private
+{
+	/** Proxy responsible for the input texture */
+	class FDMXPixelMappingRenderInputTextureProxy
+		: public IPreprocessRenderInputProxy
+	{
+	public:
+		FDMXPixelMappingRenderInputTextureProxy(UTexture* InTexture);
+
+		//~ Begin IPreprocessRenderInputProxy interface
+		virtual void Render() override;
+		virtual UTexture* GetRenderedTexture() const override;
+		virtual FVector2D GetSize2D() const override;
+		//~ End IPreprocessRenderInputProxy interface
+
+	private:
+		/** The Input object */
+		TWeakObjectPtr<UTexture> WeakInputTexture;
+	};
+}
diff --git a/Engine/Plugins/VirtualProduction/DMX/DMXPixelMapping/Source/DMXPixelMappingRenderer/Private/DMXPixelMappingRenderInputUserWidgetProxy.cpp b/Engine/Plugins/VirtualProduction/DMX/DMXPixelMapping/Source/DMXPixelMappingRenderer/Private/DMXPixelMappingRenderInputUserWidgetProxy.cpp
new file mode 100644
--- /dev/null
+++ b/Engine/Plugins/VirtualProduction/DMX/DMXPixelMapping/Source/DMXPixelMappingRenderer/Private/DMXPixelMappingRenderInputUserWidgetProxy.cpp
@@ -0,0 +1,59 @@
+// Copyright Epic Games, Inc. All Rights Reserved.
+
+#include "DMXPixelMappingRenderInputUserWidgetProxy.h"
+
+#include "Blueprint/UserWidget.h"
+#include "DMXStats.h"
+#include "Engine/TextureRenderTarget2D.h"
+#include "Framework/Application/SlateApplication.h"
+#include "Slate/WidgetRenderer.h"
+
+
+namespace UE::DMXPixelMapping::Rendering::Preprocess::Private
+{
+	DECLARE_CYCLE_STAT(TEXT("PixelMapping RenderInputUserWidget"), STAT_PreprocessRenderInputUserWidget, STATGROUP_DMX);
+
+	FDMXPixelMappingRenderInputUserWidgetProxy::FDMXPixelMappingRenderInputUserWidgetProxy(UUserWidget* InUserWidget, const FVector2D& InInputSize, EPixelFormat InFormat)
+		: WeakUserWidget(InUserWidget)
+	{
+		constexpr bool bForceLinearGamma = false;
+
+		IntermediateRenderTarget = NewObject<UTextureRenderTarget2D>();
+		IntermediateRenderTarget->ClearColor = FLinearColor::Black;
+		IntermediateRenderTarget->InitCustomFormat(InInputSize.X, InInputSize.Y, InFormat, bForceLinearGamma);
+		IntermediateRenderTarget->UpdateResourceImmediate();
+
+		const bool bUseGammaCorrection = false;
+		UMGRenderer = MakeShared<FWidgetRenderer>(bUseGammaCorrection);
+	}
+
+	void FDMXPixelMappingRenderInputUserWidgetProxy::Render()
+	{
+		SCOPE_CYCLE_COUNTER(STAT_PreprocessRenderInputUserWidget);
+
+		if (UUserWidget* UserWidget = WeakUserWidget.Get())
+		{
+			check(IntermediateRenderTarget);
+			FVector2D TextureSize = FVector2D(IntermediateRenderTarget->SizeX, IntermediateRenderTarget->SizeY);
+
+			static const float DeltaTime = 0.f;
+			UMGRenderer->DrawWidget(IntermediateRenderTarget, UserWidget->TakeWidget(), TextureSize, DeltaTime);
+		}
+	}
+
+	UTexture* FDMXPixelMappingRenderInputUserWidgetProxy::GetRenderedTexture() const
+	{
+		return IntermediateRenderTarget;
+	}
+
+	FVector2D FDMXPixelMappingRenderInputUserWidgetProxy::GetSize2D() const
+	{
+		check(IntermediateRenderTarget);
+		return FVector2D(IntermediateRenderTarget->GetSurfaceWidth(), IntermediateRenderTarget->GetSurfaceHeight());
+	}
+
+	void FDMXPixelMappingRenderInputUserWidgetProxy::AddReferencedObjects(FReferenceCollector& Collector)
+	{
+		Collector.AddReferencedObject(IntermediateRenderTarget);
+	}
+}
diff --git a/Engine/Plugins/VirtualProduction/DMX/DMXPixelMapping/Source/DMXPixelMappingRenderer/Private/DMXPixelMappingRenderInputUserWidgetProxy.h b/Engine/Plugins/VirtualProduction/DMX/DMXPixelMapping/Source/DMXPixelMappingRenderer/Private/DMXPixelMappingRenderInputUserWidgetProxy.h
new file mode 100644
--- /dev/null
+++ b/Engine/Plugins/VirtualProduction/DMX/DMXPixelMapping/Source/DMXPixelMappingRenderer/Private/DMXPixelMappingRenderInputUserWidgetProxy.h
@@ -0,0 +1,49 @@
+// Copyright Epic Games, Inc. All Rights Reserved.
+
+#pragma once
+
+#include "DMXPixelMappingPreprocessRenderer.h"
+
+#include "UObject/GCObject.h"
+
+class FWidgetRenderer;
+class UUserWidget;
+class UTextureRenderTarget2D;
+
+
+namespace UE::DMXPixelMapping::Rendering::Preprocess::Private
+{
+	/** Proxy responsible for the input user widget */
+	class FDMXPixelMappingRenderInputUserWidgetProxy
+		: public IPreprocessRenderInputProxy
+		, public FGCObject
+	{
+	public:
+		FDMXPixelMappingRenderInputUserWidgetProxy(UUserWidget* InUserWidget, const FVector2D& InInputSize, EPixelFormat InFormat);
+
+		//~ Begin IPreprocessRenderInputProxy interface
+		virtual void Render() override;
+		virtual UTexture* GetRenderedTexture() const override;
+		virtual FVector2D GetSize2D() const override;
+		//~ End IPreprocessRenderInputProxy interface
+
+	protected:
+		//~ Begin FGCObject interface
+		virtual void AddReferencedObjects(FReferenceCollector& Collector) override;
+		virtual FString GetReferencerName() const override
+		{
+			return TEXT("DMXPixelMapping::Rendering::Private::FDMXPixelMappingRenderInputUserWidgetProxy");
+		}
+		//~ End FGCObject interface
+
+	private:
+		/** UMG widget renderer */
+		TSharedPtr<FWidgetRenderer> UMGRenderer;
+
+		/** The Input object */
+		TWeakObjectPtr<UUserWidget> WeakUserWidget;
+
+		/** Intermediate render target for materials and user widgets */
+		TObjectPtr<UTextureRenderTarget2D> IntermediateRenderTarget;
+	};
+}
diff --git a/Engine/Plugins/VirtualProduction/DMX/DMXPixelMapping/Source/DMXPixelMappingRenderer/Private/DMXPixelMappingRenderPixelMapProxy.cpp b/Engine/Plugins/VirtualProduction/DMX/DMXPixelMapping/Source/DMXPixelMappingRenderer/Private/DMXPixelMappingRenderPixelMapProxy.cpp
new file mode 100644
--- /dev/null
+++ b/Engine/Plugins/VirtualProduction/DMX/DMXPixelMapping/Source/DMXPixelMappingRenderer/Private/DMXPixelMappingRenderPixelMapProxy.cpp
@@ -0,0 +1,206 @@
+// Copyright Epic Games, Inc. All Rights Reserved.
+
+#include "DMXPixelMappingRenderPixelMapProxy.h"
+
+#include "DMXPixelMappingRenderElement.h"
+#include "Engine/TextureRenderTarget2D.h"
+#include "PixelShaderUtils.h"
+#include "ProfilingDebugging/RealtimeGPUProfiler.h"
+#include "RenderingThread.h"
+#include "RHIResources.h"
+#include "TextureResource.h"
+
+
+namespace UE::DMXPixelMapping::Rendering::Private
+{
+	static constexpr TCHAR RenderPassName[] = TEXT("RenderPixelMapping");
+	static constexpr TCHAR RenderPassHint[] = TEXT("Render Pixel Mapping");
+};
+
+DECLARE_GPU_STAT_NAMED(DMXPixelMappingShadersStat, UE::DMXPixelMapping::Rendering::Private::RenderPassHint);
+
+class FDMXPixelBlendingQualityDimension : SHADER_PERMUTATION_ENUM_CLASS("PIXELBLENDING_QUALITY", EDMXPixelBlendingQuality);
+class FDMXVertexUVDimension : SHADER_PERMUTATION_BOOL("VERTEX_UV_STATIC_CALCULATION");
+
+/**
+ * Pixel Mapping downsampling vertex shader
+ */
+class FDMXPixelMappingRendererVS
+	: public FGlobalShader
+{
+public:
+	DECLARE_GLOBAL_SHADER(FDMXPixelMappingRendererVS);
+	SHADER_USE_PARAMETER_STRUCT(FDMXPixelMappingRendererVS, FGlobalShader);
+
+	using FPermutationDomain = TShaderPermutationDomain<FDMXPixelBlendingQualityDimension, FDMXVertexUVDimension>;
+	BEGIN_SHADER_PARAMETER_STRUCT(FParameters, )
+		SHADER_PARAMETER(FVector4f, DrawRectanglePosScaleBias)
+		SHADER_PARAMETER(FVector4f, DrawRectangleInvTargetSizeAndTextureSize)
+		SHADER_PARAMETER(FVector4f, DrawRectangleUVScaleBias)
+	END_SHADER_PARAMETER_STRUCT()
+
+	static bool ShouldCompilePermutation(const FGlobalShaderPermutationParameters& Parameters) { return true; }
+};
+
+/**
+ * Pixel Mapping downsampling pixel shader
+ */
+class FDMXPixelMappingRendererPS
+	: public FGlobalShader
+{
+public:
+	DECLARE_GLOBAL_SHADER(FDMXPixelMappingRendererPS);
+	SHADER_USE_PARAMETER_STRUCT(FDMXPixelMappingRendererPS, FGlobalShader);
+
+	using FPermutationDomain = TShaderPermutationDomain<FDMXPixelBlendingQualityDimension, FDMXVertexUVDimension>;
+	BEGIN_SHADER_PARAMETER_STRUCT(FParameters, )
+		SHADER_PARAMETER_TEXTURE(Texture2D, InputTexture)
+		SHADER_PARAMETER_SAMPLER(SamplerState, InputSampler)
+
+		SHADER_PARAMETER(float, Brightness)
+		SHADER_PARAMETER(FVector2f, UVTopLeftRotated)
+		SHADER_PARAMETER(FVector2f, UVTopRightRotated)
+	END_SHADER_PARAMETER_STRUCT()
+
+	static bool ShouldCompilePermutation(const FGlobalShaderPermutationParameters& Parameters) { return true; }
+};
+
+IMPLEMENT_GLOBAL_SHADER(FDMXPixelMappingRendererVS, "/Plugin/DMXPixelMapping/Private/DMXPixelMapping.usf", "DMXPixelMappingVS", SF_Vertex);
+IMPLEMENT_GLOBAL_SHADER(FDMXPixelMappingRendererPS, "/Plugin/DMXPixelMapping/Private/DMXPixelMapping.usf", "DMXPixelMappingPS", SF_Pixel);
+
+
+
+namespace UE::DMXPixelMapping::Rendering::Private
+{
+	void FDMXPixelMappingRenderPixelMapProxy::Render(
+		UTexture* InInputTexture,
+		UTextureRenderTarget2D* InRenderTarget,
+		const TArray<TSharedRef<FPixelMapRenderElement>>& InElements,
+		float Brightness
+	)
+	{
+		check(IsInGameThread());
+
+		// Ensure lifetime
+		InputTexture = InInputTexture;
+		RenderTarget = InRenderTarget;
+
+		// Actual rendering
+		ENQUEUE_RENDER_COMMAND(RenderTargetResource)(
+			[SharedThis = StaticCastSharedRef<FDMXPixelMappingRenderPixelMapProxy>(AsShared()), this, Elements = InElements, Brightness]
+			(FRHICommandListImmediate& RHICmdList)
+			{				
+				SCOPED_GPU_STAT(RHICmdList, DMXPixelMappingShadersStat);
+				SCOPED_DRAW_EVENTF(RHICmdList, DMXPixelMappingShadersStat, RenderPassName);
+
+				const FTextureResource* InputTextureResource = InputTexture ? InputTexture->GetResource() : nullptr;
+				const FTextureRenderTargetResource* RenderTargetResource = SharedThis->RenderTarget ? SharedThis->RenderTarget->GetRenderTargetResource() : nullptr;
+
+				if (!InputTextureResource || !RenderTargetResource)
+				{
+					return;
+				}
+
+				const FIntPoint PixelSize(1);
+				const FIntPoint TextureSize(1);
+
+				const FIntPoint InputTextureSize = FIntPoint(InputTextureResource->GetSizeX(), InputTextureResource->GetSizeY());
+				const FIntPoint OutputTextureSize = FIntPoint(RenderTargetResource->GetSizeX(), RenderTargetResource->GetSizeY());
+
+				const FTextureRHIRef InputTextureRHI = InputTextureResource->TextureRHI;
+
+				RHICmdList.Transition(FRHITransitionInfo(RenderTargetResource->TextureRHI, ERHIAccess::SRVMask, ERHIAccess::RTV));
+
+				FRHIRenderPassInfo RpInfo(RenderTargetResource->TextureRHI, ERenderTargetActions::Load_Store);
+				RHICmdList.BeginRenderPass(RpInfo, RenderPassName);
+				{
+					RHICmdList.SetViewport(0.f, 0.f, 0.f, OutputTextureSize.X, OutputTextureSize.Y, 1.f);
+
+					for (int32 PixelIndex = 0; PixelIndex < Elements.Num(); PixelIndex++)
+					{
+						const TSharedRef<FPixelMapRenderElement>& Element = Elements[PixelIndex];
+
+						const FVector2D& UV = Element->GetParameters().UV;
+						const FVector2D& UVSize = Element->GetParameters().UVSize;
+						const FVector2D& UVTopLeftRotated = Element->GetParameters().UVTopLeftRotated;
+						const FVector2D& UVTopRightRotated = Element->GetParameters().UVTopRightRotated;
+
+						const int32 RenderTargetPoisitionX = PixelIndex % OutputTextureSize.X;
+						const int32 RenderTargetPositionY = PixelIndex / OutputTextureSize.X;
+						
+						// Create shader permutations
+						FDMXPixelMappingRendererPS::FPermutationDomain PermutationVector;
+						PermutationVector.Set<FDMXPixelBlendingQualityDimension>(Element->GetParameters().CellBlendingQuality);
+						PermutationVector.Set<FDMXVertexUVDimension>(Element->GetParameters().bStaticCalculateUV);
+
+						// Create shaders
+						FGlobalShaderMap* ShaderMap = GetGlobalShaderMap(GMaxRHIFeatureLevel);
+						TShaderMapRef<FDMXPixelMappingRendererVS> VertexShader(ShaderMap, PermutationVector);
+						TShaderMapRef<FDMXPixelMappingRendererPS> PixelShader(ShaderMap, PermutationVector);
+
+						// Setup graphics pipeline
+						FGraphicsPipelineStateInitializer GraphicsPSOInit;
+						RHICmdList.ApplyCachedRenderTargets(GraphicsPSOInit);
+						GraphicsPSOInit.BlendState = TStaticBlendState<>::GetRHI();
+						GraphicsPSOInit.RasterizerState = TStaticRasterizerState<>::GetRHI();
+						GraphicsPSOInit.DepthStencilState = TStaticDepthStencilState<false, CF_Never>::GetRHI();
+						GraphicsPSOInit.BoundShaderState.VertexDeclarationRHI = GFilterVertexDeclaration.VertexDeclarationRHI;
+						GraphicsPSOInit.BoundShaderState.VertexShaderRHI = VertexShader.GetVertexShader();
+						GraphicsPSOInit.BoundShaderState.PixelShaderRHI = PixelShader.GetPixelShader();
+						GraphicsPSOInit.PrimitiveType = PT_TriangleList;
+						SetGraphicsPipelineState(RHICmdList, GraphicsPSOInit, 0);
+
+						// Set vertex shader buffer
+						FDMXPixelMappingRendererVS::FParameters VSParameters;
+						VSParameters.DrawRectanglePosScaleBias = FVector4f(PixelSize.X, PixelSize.Y, RenderTargetPoisitionX, RenderTargetPositionY);
+						VSParameters.DrawRectangleUVScaleBias = FVector4f(UVSize.X, UVSize.Y, UV.X, UV.Y);
+						VSParameters.DrawRectangleInvTargetSizeAndTextureSize = FVector4f(
+							1.0 / OutputTextureSize.X, 1.0 / OutputTextureSize.Y,
+							1.0 / TextureSize.X, 1.0 / TextureSize.Y);
+						SetShaderParameters(RHICmdList, VertexShader, VertexShader.GetVertexShader(), VSParameters);
+
+						// Set pixel shader buffer
+						FDMXPixelMappingRendererPS::FParameters PSParameters;
+						PSParameters.InputTexture = InputTextureRHI;
+						PSParameters.Brightness = Brightness;
+						PSParameters.InputSampler = TStaticSamplerState<SF_Bilinear, AM_Clamp, AM_Clamp, AM_Clamp>::GetRHI();
+						PSParameters.UVTopLeftRotated = FVector2f(UVTopLeftRotated);
+						PSParameters.UVTopRightRotated = FVector2f(UVTopRightRotated);
+						SetShaderParameters(RHICmdList, PixelShader, PixelShader.GetPixelShader(), PSParameters);
+
+						// Draw a two triangle on the entire viewport.
+						FPixelShaderUtils::DrawFullscreenQuad(RHICmdList, 1);
+					}
+				}
+				RHICmdList.EndRenderPass();
+
+				// Copy texture from GPU to CPU
+				{
+					// Read the contents of a texture to an output CPU buffer
+					TArray<FLinearColor> ColorArray;
+					const FIntRect Rect(0, 0, OutputTextureSize.X, OutputTextureSize.Y);
+
+					// Read surface data. 
+					// Use min/max compression mode here - Either the input was tone mapped earlier, or we want to preserve dynamic range.
+					RHICmdList.ReadSurfaceData(RenderTargetResource->TextureRHI, Rect, ColorArray, ERangeCompressionMode::RCM_MinMax);
+
+					// Write data to the game thread
+					if (!ensureMsgf(ColorArray.Num() >= Elements.Num(), TEXT("Pixel maparray is smaller than the number of elements that should have been rendered. Failed to Pixel Map")))
+					{
+						return;
+					}
+
+					for (int32 PixelIndex = 0; PixelIndex < Elements.Num(); PixelIndex++)
+					{
+						Elements[PixelIndex]->SetColor(ColorArray[PixelIndex]);
+					}
+				}
+			});
+	}
+
+	void FDMXPixelMappingRenderPixelMapProxy::AddReferencedObjects(FReferenceCollector& Collector)
+	{
+		Collector.AddReferencedObject(InputTexture);
+		Collector.AddReferencedObject(RenderTarget);
+	}
+}
diff --git a/Engine/Plugins/VirtualProduction/DMX/DMXPixelMapping/Source/DMXPixelMappingRenderer/Private/DMXPixelMappingRenderPixelMapProxy.h b/Engine/Plugins/VirtualProduction/DMX/DMXPixelMapping/Source/DMXPixelMappingRenderer/Private/DMXPixelMappingRenderPixelMapProxy.h
new file mode 100644
--- /dev/null
+++ b/Engine/Plugins/VirtualProduction/DMX/DMXPixelMapping/Source/DMXPixelMappingRenderer/Private/DMXPixelMappingRenderPixelMapProxy.h
@@ -0,0 +1,50 @@
+// Copyright Epic Games, Inc. All Rights Reserved.
+
+#pragma once
+
+#include "DMXPixelMappingPixelMapRenderer.h"
+#include "UObject/GCObject.h"
+
+class UTexture;
+class UTextureRenderTarget2D;
+namespace UE::DMXPixelMapping::Rendering { class FPixelMapRenderElement; }
+
+
+namespace UE::DMXPixelMapping::Rendering::Private
+{
+	/** Interface for the object responsible to render the input texture/material/umg widget. */
+	class FDMXPixelMappingRenderPixelMapProxy
+		: public IRenderPixelMapProxy
+		, public FGCObject
+	{
+	public:
+		/**
+		 * Pixelmapping specific, downsample and draw input texture to destination texture.
+		 *
+		 * @param InInputTexture				The input texture to pixel map
+		 * @param InRenderTarget				The render target that is pixel mapped to
+		 * @param InElements					Elements to be rendered
+		 * @param InBrightness					Global brightness
+		 * @param InCallback					Callback for reading  the pixels from GPU to CPU
+		 */
+		virtual void Render(
+			UTexture* InInputTexture,
+			UTextureRenderTarget2D* InRenderTarget,
+			const TArray<TSharedRef<FPixelMapRenderElement>>& InElements,
+			float InBrightness
+		) override;
+
+	protected:
+		//~ Begin FGCObject interface
+		virtual void AddReferencedObjects(FReferenceCollector& Collector) override;
+		virtual FString GetReferencerName() const override { return TEXT("FDMXPixelMappingRenderPixelMapProxy"); }
+		//~ End FGCObject interface
+
+	private:
+		/** Strong ref to the current input texture */
+		TObjectPtr<UTexture> InputTexture;
+
+		/** Strong ref to the current render target */
+		TObjectPtr<UTextureRenderTarget2D> RenderTarget;
+	};
+}
diff --git a/Engine/Plugins/VirtualProduction/DMX/DMXPixelMapping/Source/DMXPixelMappingRenderer/Private/PreprocessApplyFilterMaterialProxy.cpp b/Engine/Plugins/VirtualProduction/DMX/DMXPixelMapping/Source/DMXPixelMappingRenderer/Private/PreprocessApplyFilterMaterialProxy.cpp
deleted file mode 100644
diff --git a/Engine/Plugins/VirtualProduction/DMX/DMXPixelMapping/Source/DMXPixelMappingRenderer/Private/PreprocessApplyFilterMaterialProxy.h b/Engine/Plugins/VirtualProduction/DMX/DMXPixelMapping/Source/DMXPixelMappingRenderer/Private/PreprocessApplyFilterMaterialProxy.h
deleted file mode 100644
diff --git a/Engine/Plugins/VirtualProduction/DMX/DMXPixelMapping/Source/DMXPixelMappingRenderer/Private/PreprocessRenderInputMaterialProxy .cpp b/Engine/Plugins/VirtualProduction/DMX/DMXPixelMapping/Source/DMXPixelMappingRenderer/Private/PreprocessRenderInputMaterialProxy .cpp
deleted file mode 100644
diff --git a/Engine/Plugins/VirtualProduction/DMX/DMXPixelMapping/Source/DMXPixelMappingRenderer/Private/PreprocessRenderInputMaterialProxy.h b/Engine/Plugins/VirtualProduction/DMX/DMXPixelMapping/Source/DMXPixelMappingRenderer/Private/PreprocessRenderInputMaterialProxy.h
deleted file mode 100644
diff --git a/Engine/Plugins/VirtualProduction/DMX/DMXPixelMapping/Source/DMXPixelMappingRenderer/Private/PreprocessRenderInputTextureProxy.cpp b/Engine/Plugins/VirtualProduction/DMX/DMXPixelMapping/Source/DMXPixelMappingRenderer/Private/PreprocessRenderInputTextureProxy.cpp
deleted file mode 100644
diff --git a/Engine/Plugins/VirtualProduction/DMX/DMXPixelMapping/Source/DMXPixelMappingRenderer/Private/PreprocessRenderInputTextureProxy.h b/Engine/Plugins/VirtualProduction/DMX/DMXPixelMapping/Source/DMXPixelMappingRenderer/Private/PreprocessRenderInputTextureProxy.h
deleted file mode 100644
diff --git a/Engine/Plugins/VirtualProduction/DMX/DMXPixelMapping/Source/DMXPixelMappingRenderer/Private/PreprocessRenderInputUserWidgetProxy.cpp b/Engine/Plugins/VirtualProduction/DMX/DMXPixelMapping/Source/DMXPixelMappingRenderer/Private/PreprocessRenderInputUserWidgetProxy.cpp
deleted file mode 100644
diff --git a/Engine/Plugins/VirtualProduction/DMX/DMXPixelMapping/Source/DMXPixelMappingRenderer/Private/PreprocessRenderInputUserWidgetProxy.h b/Engine/Plugins/VirtualProduction/DMX/DMXPixelMapping/Source/DMXPixelMappingRenderer/Private/PreprocessRenderInputUserWidgetProxy.h
deleted file mode 100644
diff --git a/Engine/Plugins/VirtualProduction/DMX/DMXPixelMapping/Source/DMXPixelMappingRenderer/Private/RenderPixelMapProxy.cpp b/Engine/Plugins/VirtualProduction/DMX/DMXPixelMapping/Source/DMXPixelMappingRenderer/Private/RenderPixelMapProxy.cpp
deleted file mode 100644
diff --git a/Engine/Plugins/VirtualProduction/DMX/DMXPixelMapping/Source/DMXPixelMappingRenderer/Private/RenderPixelMapProxy.h b/Engine/Plugins/VirtualProduction/DMX/DMXPixelMapping/Source/DMXPixelMappingRenderer/Private/RenderPixelMapProxy.h
deleted file mode 100644
