DMX - Fix matrices in Pixel Mapping do not update when moved, when fixture type or patch changed, matrix cells can be dragged when not locked in designer

diff --git a/Engine/Plugins/VirtualProduction/DMX/DMXPixelMapping/Source/DMXPixelMappingEditor/Private/Views/SDMXPixelMappingDesignerView.cpp b/Engine/Plugins/VirtualProduction/DMX/DMXPixelMapping/Source/DMXPixelMappingEditor/Private/Views/SDMXPixelMappingDesignerView.cpp
--- a/Engine/Plugins/VirtualProduction/DMX/DMXPixelMapping/Source/DMXPixelMappingEditor/Private/Views/SDMXPixelMappingDesignerView.cpp
+++ b/Engine/Plugins/VirtualProduction/DMX/DMXPixelMapping/Source/DMXPixelMappingEditor/Private/Views/SDMXPixelMappingDesignerView.cpp
@@ -963,8 +963,9 @@ void SDMXPixelMappingDesignerView::ResolvePendingSelectedComponents(bool bClearP
 	const bool bHasPendingSelectedComponent = PendingSelectedComponent.IsValid();
 	if (Toolkit.IsValid() && bHasPendingSelectedComponent)
 	{
-		// Never select matrix cells
-		if (UDMXPixelMappingMatrixCellComponent* MatrixCellComponent = Cast<UDMXPixelMappingMatrixCellComponent>(PendingSelectedComponent))
+		// Select matrix cells only if they're not locked in designer. By default they're locked.
+		UDMXPixelMappingMatrixCellComponent* MatrixCellComponent = Cast<UDMXPixelMappingMatrixCellComponent>(PendingSelectedComponent);
+		if (MatrixCellComponent && MatrixCellComponent->IsLockInDesigner())
 		{
 			PendingSelectedComponent = MatrixCellComponent->GetParent();
 		}
diff --git a/Engine/Plugins/VirtualProduction/DMX/DMXPixelMapping/Source/DMXPixelMappingRuntime/Private/Components/DMXPixelMappingMatrixComponent.cpp b/Engine/Plugins/VirtualProduction/DMX/DMXPixelMapping/Source/DMXPixelMappingRuntime/Private/Components/DMXPixelMappingMatrixComponent.cpp
--- a/Engine/Plugins/VirtualProduction/DMX/DMXPixelMapping/Source/DMXPixelMappingRuntime/Private/Components/DMXPixelMappingMatrixComponent.cpp
+++ b/Engine/Plugins/VirtualProduction/DMX/DMXPixelMapping/Source/DMXPixelMappingRuntime/Private/Components/DMXPixelMappingMatrixComponent.cpp
@@ -392,6 +392,8 @@ void UDMXPixelMappingMatrixComponent::SetPosition(const FVector2D& NewPosition)
 #if WITH_EDITOR
 	EditorPositionWithRotation = CachedGeometry.GetPositionRotatedAbsolute();
 #endif
+
+	InvalidatePixelMapRenderer();
 }
 
 void UDMXPixelMappingMatrixComponent::SetPositionRotated(FVector2D NewRotatedPosition)
@@ -406,6 +408,8 @@ void UDMXPixelMappingMatrixComponent::SetPositionRotated(FVector2D NewRotatedPos
 #if WITH_EDITOR
 	EditorPositionWithRotation = CachedGeometry.GetPositionRotatedAbsolute();
 #endif
+
+	InvalidatePixelMapRenderer();
 }
 
 void UDMXPixelMappingMatrixComponent::SetSize(const FVector2D& NewSize)
@@ -422,6 +426,8 @@ void UDMXPixelMappingMatrixComponent::SetSize(const FVector2D& NewSize)
 	SetPosition(CachedGeometry.GetPositionAbsolute());
 
 	CachedGeometry.PropagonateSizeChangesToChildren(DeltaSize, DeltaPosition);
+
+	InvalidatePixelMapRenderer();
 }
 
 void UDMXPixelMappingMatrixComponent::SetRotation(double NewRotation)
@@ -434,6 +440,8 @@ void UDMXPixelMappingMatrixComponent::SetRotation(double NewRotation)
 #if WITH_EDITOR
 	EditorPositionWithRotation = CachedGeometry.GetPositionRotatedAbsolute();
 #endif
+
+	InvalidatePixelMapRenderer();
 }
 
 PRAGMA_DISABLE_DEPRECATION_WARNINGS
@@ -515,6 +523,8 @@ void UDMXPixelMappingMatrixComponent::HandleMatrixChanged()
 
 	SetRotation(RestoreRotation);
 
+	InvalidatePixelMapRenderer();
+
 	PRAGMA_DISABLE_DEPRECATION_WARNINGS
 	GetOnMatrixChanged().Broadcast(PixelMapping, this);
 	PRAGMA_ENABLE_DEPRECATION_WARNINGS
