DMX - Fix CIS Build Engine Localization, duplicate key "FlipHorizontalText"

diff --git a/Engine/Plugins/VirtualProduction/DMX/DMXPixelMapping/Source/DMXPixelMappingEditor/Private/Toolkits/DMXPixelMappingToolkit.cpp b/Engine/Plugins/VirtualProduction/DMX/DMXPixelMapping/Source/DMXPixelMappingEditor/Private/Toolkits/DMXPixelMappingToolkit.cpp
--- a/Engine/Plugins/VirtualProduction/DMX/DMXPixelMapping/Source/DMXPixelMappingEditor/Private/Toolkits/DMXPixelMappingToolkit.cpp
+++ b/Engine/Plugins/VirtualProduction/DMX/DMXPixelMapping/Source/DMXPixelMappingEditor/Private/Toolkits/DMXPixelMappingToolkit.cpp
@@ -603,7 +603,7 @@ void FDMXPixelMappingToolkit::FlipGroup(EOrientation Orientation, bool bTransact
 	{
 		const FText TransactionText = FText::Format(LOCTEXT("FilpCellsTransaction", "Flip Group {0}"), Orientation == EOrientation::Orient_Horizontal ? 
 			LOCTEXT("FlipHorizontalText", "Horizontally") :
-			LOCTEXT("FlipHorizontalText", "Vertically"));
+			LOCTEXT("FlipVerticalText", "Vertically"));
 
 		FilpCellsTransaction = MakeShared<FScopedTransaction>(TransactionText);
 	}
