DMX - Fix designer shows components offset on high zoom levels, when OS DPI scaling is enabled. Fix an issue where matrix components do not correctly follow size when related option in renderer is enabled

diff --git a/Engine/Plugins/VirtualProduction/DMX/DMXPixelMapping/Source/DMXPixelMappingEditor/Private/DMXPixelMappingEditorCommands.cpp b/Engine/Plugins/VirtualProduction/DMX/DMXPixelMapping/Source/DMXPixelMappingEditor/Private/DMXPixelMappingEditorCommands.cpp
--- a/Engine/Plugins/VirtualProduction/DMX/DMXPixelMapping/Source/DMXPixelMappingEditor/Private/DMXPixelMappingEditorCommands.cpp
+++ b/Engine/Plugins/VirtualProduction/DMX/DMXPixelMapping/Source/DMXPixelMappingEditor/Private/DMXPixelMappingEditorCommands.cpp
@@ -20,15 +20,15 @@ void FDMXPixelMappingEditorCommands::RegisterCommands()
 {
 	UI_COMMAND(AddMapping, "Add Source", "Adds a new Source Texure, Material or User Widget to the Pixel Map asset.", EUserInterfaceActionType::Button, FInputChord());
 	UI_COMMAND(PlayDMX, "Plays DMX", "Plays DMX", EUserInterfaceActionType::Button, FInputChord());
-	UI_COMMAND(PauseDMX, "Pause DMX", "Pauses playing DMX. DMX values will still be sent, at a lower rate.", EUserInterfaceActionType::Button, FInputChord());
-	UI_COMMAND(ResumeDMX, "Resume DMX", "Resumes playing DMX after being paused. ", EUserInterfaceActionType::Button, FInputChord());
+	UI_COMMAND(PauseDMX, "Pause DMX", "Pauses playing DMX", EUserInterfaceActionType::Button, FInputChord());
+	UI_COMMAND(ResumeDMX, "Resume DMX", "Resumes playing DMX after being paused", EUserInterfaceActionType::Button, FInputChord());
 	UI_COMMAND(StopDMX, "Stop DMX", "Stops playing DMX.", EUserInterfaceActionType::Button, FInputChord());
 	UI_COMMAND(TogglePlayPauseDMX, "Toggle Play/Pause DMX", "Toggles between playing and pausing DMX", EUserInterfaceActionType::Button, FInputChord(EModifierKey::Shift, EKeys::SpaceBar));
 	UI_COMMAND(TogglePlayStopDMX, "Toggle Play/Stop DMX", "Toggles between playing and stopping DMX", EUserInterfaceActionType::Button, FInputChord(EKeys::SpaceBar));
 
 	UI_COMMAND(EditorStopSendsZeroValues, "Stop sends Zero Values", "When stop is clicked in the editor, zero values are sent to all patches in use", EUserInterfaceActionType::RadioButton, FInputChord());
 	UI_COMMAND(EditorStopSendsDefaultValues, "Stop sends Default Values", "When stop is clicked in the editor, default values are sent to all patches in use", EUserInterfaceActionType::RadioButton, FInputChord());
-	UI_COMMAND(EditorStopKeepsLastValues, "Stop keeps last Values", "When stop is clicked in the editor, the last pixel mapped values are kept", EUserInterfaceActionType::RadioButton, FInputChord());
+	UI_COMMAND(EditorStopKeepsLastValues, "Stop keeps last Values", "When stop is clicked in the editor, pixel mapping leaves last send values untouched", EUserInterfaceActionType::RadioButton, FInputChord());
 
 	UI_COMMAND(EnableResizeMode, "Resize Mode", "Resizes components when transform handles are being dragged", EUserInterfaceActionType::RadioButton, FInputChord());
 	UI_COMMAND(EnableRotateMode, "Rotate Mode", "Rotates components when transform handles are being dragged", EUserInterfaceActionType::RadioButton, FInputChord());
diff --git a/Engine/Plugins/VirtualProduction/DMX/DMXPixelMapping/Source/DMXPixelMappingEditor/Private/Viewports/DMXPixelMappingSourceTextureViewportClient.cpp b/Engine/Plugins/VirtualProduction/DMX/DMXPixelMapping/Source/DMXPixelMappingEditor/Private/Viewports/DMXPixelMappingSourceTextureViewportClient.cpp
--- a/Engine/Plugins/VirtualProduction/DMX/DMXPixelMapping/Source/DMXPixelMappingEditor/Private/Viewports/DMXPixelMappingSourceTextureViewportClient.cpp
+++ b/Engine/Plugins/VirtualProduction/DMX/DMXPixelMapping/Source/DMXPixelMappingEditor/Private/Viewports/DMXPixelMappingSourceTextureViewportClient.cpp
@@ -120,15 +120,19 @@ void FDMXPixelMappingSourceTextureViewportClient::Draw(FViewport* InViewport, FC
 		{
 			return;
 		}
-
+	
 		FCanvasTileItem TileItem(FVector2D::ZeroVector, InputTexture->GetResource(), SceneViewport->GetSizeXY(), UV0, UV1, ColorWithExposure);
 		TileItem.BlendMode = ESimpleElementBlendMode::SE_BLEND_MAX;
+
+		bUseDPIScaling = false;
 		Canvas->DrawItem(TileItem);
 	}
 	else
 	{
 		FCanvasTileItem TileItem(VisibleRect.Min * DesignerView->GetZoomAmount(), InputTexture->GetResource(), VisibleRect.GetSize() * DesignerView->GetZoomAmount(), UV0, UV1, ColorWithExposure);
 		TileItem.BlendMode = ESimpleElementBlendMode::SE_BLEND_MAX;
+		
+		bUseDPIScaling = true;
 		Canvas->DrawItem(TileItem);
 	}
 }
diff --git a/Engine/Plugins/VirtualProduction/DMX/DMXPixelMapping/Source/DMXPixelMappingEditor/Private/Viewports/DMXPixelMappingSourceTextureViewportClient.h b/Engine/Plugins/VirtualProduction/DMX/DMXPixelMapping/Source/DMXPixelMappingEditor/Private/Viewports/DMXPixelMappingSourceTextureViewportClient.h
--- a/Engine/Plugins/VirtualProduction/DMX/DMXPixelMapping/Source/DMXPixelMappingEditor/Private/Viewports/DMXPixelMappingSourceTextureViewportClient.h
+++ b/Engine/Plugins/VirtualProduction/DMX/DMXPixelMapping/Source/DMXPixelMappingEditor/Private/Viewports/DMXPixelMappingSourceTextureViewportClient.h
@@ -24,6 +24,7 @@ public:
 protected:
 	//~ Begin FViewportClient Interface
 	virtual void Draw(FViewport* Viewport, FCanvas* Canvas) override;
+	virtual bool ShouldDPIScaleSceneCanvas() const override { return bUseDPIScaling; }
 	//~ End FViewportClient Interface
 
 	//~ Begin FViewport Interface
@@ -31,6 +32,8 @@ protected:
 	//~ End FViewport Interface
 
 private:
+	bool bUseDPIScaling = false;
+
 	TWeakPtr<FDMXPixelMappingToolkit> WeakToolkit;
 	TWeakPtr<SDMXPixelMappingSourceTextureViewport> WeakSourceTextureViewport;
 };
diff --git a/Engine/Plugins/VirtualProduction/DMX/DMXPixelMapping/Source/DMXPixelMappingRuntime/Private/Components/DMXPixelMappingRendererComponent.cpp b/Engine/Plugins/VirtualProduction/DMX/DMXPixelMapping/Source/DMXPixelMappingRuntime/Private/Components/DMXPixelMappingRendererComponent.cpp
--- a/Engine/Plugins/VirtualProduction/DMX/DMXPixelMapping/Source/DMXPixelMappingRuntime/Private/Components/DMXPixelMappingRendererComponent.cpp
+++ b/Engine/Plugins/VirtualProduction/DMX/DMXPixelMapping/Source/DMXPixelMappingRuntime/Private/Components/DMXPixelMappingRendererComponent.cpp
@@ -375,13 +375,19 @@ void UDMXPixelMappingRendererComponent::LetChildrenFollowSize()
 		{
 			if (UDMXPixelMappingOutputComponent* OutputComponent = Cast<UDMXPixelMappingOutputComponent>(Component))
 			{
-				OutputComponent->Modify();
+				if (OutputComponent->GetClass() == UDMXPixelMappingFixtureGroupComponent::StaticClass() ||
+					OutputComponent->GetClass() == UDMXPixelMappingFixtureGroupItemComponent::StaticClass() ||
+					OutputComponent->GetClass() == UDMXPixelMappingMatrixComponent::StaticClass())
+				{
+					OutputComponent->Modify();
+
+					const FVector2D NewPosition = OutputComponent->GetPosition() * Scalar;
+					OutputComponent->SetPosition(NewPosition);
 
-				const FVector2D NewPosition = OutputComponent->GetPosition() * Scalar;
-				OutputComponent->SetPosition(NewPosition);
+					const FVector2D NewSize = OutputComponent->GetSize() * Scalar;
+					OutputComponent->SetSize(NewSize);
+				}
 
-				const FVector2D NewSize = OutputComponent->GetSize() * Scalar;
-				OutputComponent->SetSize(NewSize);
 			}
 		}, 
 		bRecursive);
