DMX - Fix PixelMapping designer editor exposure setting no longer was applied as expected after unrelated engine changes (u

diff --git a/Engine/Plugins/VirtualProduction/DMX/DMXPixelMapping/Source/DMXPixelMappingEditor/Private/Viewports/DMXPixelMappingSourceTextureViewportClient.cpp b/Engine/Plugins/VirtualProduction/DMX/DMXPixelMapping/Source/DMXPixelMappingEditor/Private/Viewports/DMXPixelMappingSourceTextureViewportClient.cpp
--- a/Engine/Plugins/VirtualProduction/DMX/DMXPixelMapping/Source/DMXPixelMappingEditor/Private/Viewports/DMXPixelMappingSourceTextureViewportClient.cpp
+++ b/Engine/Plugins/VirtualProduction/DMX/DMXPixelMapping/Source/DMXPixelMappingEditor/Private/Viewports/DMXPixelMappingSourceTextureViewportClient.cpp
@@ -112,7 +112,7 @@ void FDMXPixelMappingSourceTextureViewportClient::Draw(FViewport* InViewport, FC
 	const FVector2D UV0 = VisibleRect.Min / TextureBox.GetSize();
 	const FVector2D UV1 = VisibleRect.Max / TextureBox.GetSize();
 
-	const FLinearColor ColorWithExposure = FLinearColor::White * PixelMapping->DesignerExposure;
+	const FLinearColor ColorWithExposure = (FLinearColor::White * PixelMapping->DesignerExposure).CopyWithNewOpacity(1.f);
 	if (IsDrawingVisibleRectOnly())
 	{
 		const TSharedPtr<FSceneViewport> SceneViewport = ViewportWidget->GetViewport();
