DMX - Fix PixelMapping fixture rotation is not resetted when layout script is run, marginally improve UX

diff --git a/Engine/Plugins/VirtualProduction/DMX/DMXPixelMapping/Source/DMXPixelMappingEditor/Private/Customizations/DMXPixelMappingLayoutViewModelDetails.cpp b/Engine/Plugins/VirtualProduction/DMX/DMXPixelMapping/Source/DMXPixelMappingEditor/Private/Customizations/DMXPixelMappingLayoutViewModelDetails.cpp
--- a/Engine/Plugins/VirtualProduction/DMX/DMXPixelMapping/Source/DMXPixelMappingEditor/Private/Customizations/DMXPixelMappingLayoutViewModelDetails.cpp
+++ b/Engine/Plugins/VirtualProduction/DMX/DMXPixelMapping/Source/DMXPixelMappingEditor/Private/Customizations/DMXPixelMappingLayoutViewModelDetails.cpp
@@ -2,11 +2,16 @@
 
 #include "Customizations/DMXPixelMappingLayoutViewModelDetails.h"
 
-#include "ViewModels/DMXPixelMappingLayoutViewModel.h"
-
+#include "Components/DMXPixelMappingOutputComponent.h"
 #include "DetailLayoutBuilder.h"
 #include "DetailWidgetRow.h"
 #include "IPropertyUtilities.h"
+#include "ScopedTransaction.h"
+#include "Settings/DMXPixelMappingEditorSettings.h"
+#include "Toolkits/DMXPixelMappingToolkit.h"
+#include "ViewModels/DMXPixelMappingLayoutViewModel.h"
+#include "Widgets/Input/SButton.h"
+#include "Widgets/Input/SCheckBox.h"
 #include "Widgets/Text/STextBlock.h"
 
 
@@ -20,48 +25,91 @@ TSharedRef<IDetailCustomization> FDMXPixelMappingLayoutViewModelDetails::MakeIns
 void FDMXPixelMappingLayoutViewModelDetails::CustomizeDetails(IDetailLayoutBuilder& DetailLayout)
 {
 	PropertyUtilities = DetailLayout.GetPropertyUtilities();
+
+	const TSharedPtr<FDMXPixelMappingToolkit> Toolkit = WeakToolkit.Pin();
+	UDMXPixelMappingLayoutViewModel* ViewModel = GetLayoutViewModel();
+	if (!Toolkit.IsValid() || !ViewModel)
+	{
+		return;
+	}
 	
 	const TSharedPtr<IPropertyHandle> LayoutScriptClassHandle = DetailLayout.GetProperty(UDMXPixelMappingLayoutViewModel::GetLayoutScriptClassPropertyName());
-	
-	// Show info about which components are being laid out, or why no layout can be applied
-	DetailLayout.EditCategory(LayoutScriptClassHandle->GetDefaultCategoryName())
-		.AddCustomRow(FText::GetEmpty())
+	IDetailCategoryBuilder& LayoutScriptCategory = DetailLayout.EditCategory(LayoutScriptClassHandle->GetDefaultCategoryName());
+
+	// Add a checkbox to apply layout scripts instantly
+	const FText ApplyInstantlyTooltip = LOCTEXT("ApplyLayoutScriptWhenLoadedTooltip", "When checked, applies the Layout Script instantly when a new script is selected.");
+	LayoutScriptCategory.AddCustomRow(FText::GetEmpty())
 		.NameContent()
 		[
 			SNew(STextBlock)
-			.Text(LOCTEXT("LayouotModeLabel", "Layout Mode"))
+			.Text(LOCTEXT("ApplyLayoutScriptWhenLoadedLabel", "Apply New Scripts Instantly"))
+			.ToolTipText(ApplyInstantlyTooltip)
 			.Font(FAppStyle::Get().GetFontStyle("PropertyWindow.NormalFont"))
 		]
 		.ValueContent()
 		[
-			SNew(STextBlock)
-			.Text_Lambda([this]()
+			SNew(SCheckBox)
+			.ToolTipText(ApplyInstantlyTooltip)
+			.IsChecked_Lambda([]()
 				{
-					const UDMXPixelMappingLayoutViewModel* const Model = GetLayoutViewModel();
-					if (!Model || Model->GetMode() == EDMXPixelMappingLayoutViewModelMode::LayoutNone)
-					{
-						return LOCTEXT("NoLayoutCanBeAppliedInfo", "None");
-					}
-					else if (Model->GetMode() == EDMXPixelMappingLayoutViewModelMode::LayoutRendererComponentChildren)
-					{
-						return LOCTEXT("LayoutGroupsInfo", "Groups");
-					}
-					else if (Model->GetMode() == EDMXPixelMappingLayoutViewModelMode::LayoutFixtureGroupComponentChildren)
-					{
-						return LOCTEXT("LayoutGroupChildrenInfo", "Childs");
-					}
-					else if (Model->GetMode() == EDMXPixelMappingLayoutViewModelMode::LayoutMatrixComponentChildren)
-					{
-						return LOCTEXT("LayoutMatrixCellsInfo", "Matrix Cells");
-					}
-					return FText::GetEmpty();
+					return GetDefault<UDMXPixelMappingEditorSettings>()->DesignerSettings.bApplyLayoutScriptWhenLoaded ?
+						ECheckBoxState::Checked :
+						ECheckBoxState::Unchecked;
 				})
-			.Font(FAppStyle::Get().GetFontStyle("PropertyWindow.NormalFont"))
+			.OnCheckStateChanged_Lambda([](ECheckBoxState NewCheckboxState)
+			{
+				UDMXPixelMappingEditorSettings* EditorSettings = GetMutableDefault<UDMXPixelMappingEditorSettings>();
+				EditorSettings->DesignerSettings.bApplyLayoutScriptWhenLoaded = NewCheckboxState == ECheckBoxState::Checked;
+				EditorSettings->SaveConfig();
+			})
 		];
 
-	// Hide the Layout Script Class property by default, then show it depending on the layout view model mode.
-	LayoutScriptClassHandle->MarkHiddenByCustomization();
+	// Add the auto apply property
+	AutoApplyHandle = DetailLayout.GetProperty(UDMXPixelMappingLayoutViewModel::GetAutoApplyPropertyName());
+	LayoutScriptCategory.AddProperty(AutoApplyHandle)
+		.Visibility(TAttribute<EVisibility>::CreateSP(this, &FDMXPixelMappingLayoutViewModelDetails::GetAutoApplyPropertyVisiblity));
+
+	// Add an apply button
+	LayoutScriptCategory.AddCustomRow(FText::GetEmpty())
+		.Visibility(TAttribute<EVisibility>::CreateSP(this, &FDMXPixelMappingLayoutViewModelDetails::GetApplyLayoutScriptButtonVisibility))
+		.WholeRowContent()
+		.HAlign(HAlign_Left)
+		[
+			SNew(SButton)
+			.ButtonStyle(FAppStyle::Get(), "FlatButton.Success")
+			.ForegroundColor(FLinearColor::White)
+			.ContentPadding(FMargin(5.f, 1.f))
+			.OnClicked(this, &FDMXPixelMappingLayoutViewModelDetails::OnApplyLayoutScriptClicked)
+			.Content()
+			[
+				SNew(SBorder)
+				.BorderImage(FAppStyle::GetBrush("NoBorder"))
+				.HAlign(HAlign_Center)
+				.VAlign(VAlign_Center)
+				[
+					SNew(STextBlock)
+					.Text(LOCTEXT("ApplyLayoutScriptLabel", "Apply Script"))
+				]
+			]
+		];
 
+	// Show info about which components are being laid out, or why no layout can be applied
+	LayoutScriptCategory
+		.AddCustomRow(FText::GetEmpty())
+		.NameContent()
+		[
+			SNew(STextBlock)
+			.Text(LOCTEXT("LayouotModeLabel", "Applied to"))
+			.Font(FAppStyle::Get().GetFontStyle("PropertyWindow.NormalFont"))
+		]
+		.ValueContent()
+		[
+			SAssignNew(InfoTextBlock, STextBlock)
+			.Font(FAppStyle::Get().GetFontStyle("PropertyWindow.NormalFont"))
+		];
+	UpdateInfoText();
+
+	// Show the Layout Script Class property depending on the layout view model mode.
 	const TAttribute<EVisibility> LayoutScriptClassVisibilityAttribute = TAttribute<EVisibility>::CreateLambda([this]()
 		{
 			const UDMXPixelMappingLayoutViewModel* const Model = GetLayoutViewModel();
@@ -72,9 +120,87 @@ void FDMXPixelMappingLayoutViewModelDetails::CustomizeDetails(IDetailLayoutBuild
 			return EVisibility::Collapsed;	
 		});
 
-	DetailLayout.EditCategory(LayoutScriptClassHandle->GetDefaultCategoryName())
+	LayoutScriptCategory
 		.AddProperty(LayoutScriptClassHandle)
 		.Visibility(LayoutScriptClassVisibilityAttribute);
+
+	// Update info on changes
+	ViewModel->OnModelChanged.AddSP(this, &FDMXPixelMappingLayoutViewModelDetails::UpdateInfoText);
+	Toolkit->GetOnSelectedComponentsChangedDelegate().AddSP(this, &FDMXPixelMappingLayoutViewModelDetails::UpdateInfoText);
+}
+
+EVisibility FDMXPixelMappingLayoutViewModelDetails::GetAutoApplyPropertyVisiblity() const
+{
+	if (UDMXPixelMappingLayoutViewModel* LayoutViewModel = GetLayoutViewModel())
+	{
+		return LayoutViewModel->CanApplyLayoutScript() ? EVisibility::Visible : EVisibility::Collapsed;
+	}
+
+	return EVisibility::Collapsed;
+}
+
+EVisibility FDMXPixelMappingLayoutViewModelDetails::GetApplyLayoutScriptButtonVisibility() const
+{
+	bool bAutoApply;
+	if (AutoApplyHandle->GetValue(bAutoApply) == FPropertyAccess::Success &&
+		!bAutoApply)
+	{
+		if (UDMXPixelMappingLayoutViewModel* LayoutViewModel = GetLayoutViewModel())
+		{
+			return LayoutViewModel->CanApplyLayoutScript() ? EVisibility::Visible : EVisibility::Collapsed;
+		}
+	}
+
+	return EVisibility::Collapsed;
+}
+
+FReply FDMXPixelMappingLayoutViewModelDetails::OnApplyLayoutScriptClicked()
+{
+	if (UDMXPixelMappingLayoutViewModel* LayoutViewModel = GetLayoutViewModel())
+	{
+		const FScopedTransaction ApplyLayoutScriptTransaction(LOCTEXT("ApplyLayoutScriptTransaction", "Apply Layout Script"));
+		LayoutViewModel->ForceApplyLayoutScript();
+	}
+
+	return FReply::Handled();
+}
+
+void FDMXPixelMappingLayoutViewModelDetails::UpdateInfoText()
+{
+	if (UDMXPixelMappingLayoutViewModel* LayoutViewModel = GetLayoutViewModel())
+	{
+		UDMXPixelMappingOutputComponent* ParentComponent = LayoutViewModel->GetParentComponent();
+
+		const FText InfoText = [ParentComponent, this]()
+			{		
+				const FText ParentComponentNameText = ParentComponent ?
+					FText::FromString(ParentComponent->GetUserName()) :
+					LOCTEXT("InvalidParentName", "Invalid Parent");
+
+				const FText ChildrenOfText = LOCTEXT("LayoutChildrenOfText", "Children of {0}");
+
+				const UDMXPixelMappingLayoutViewModel* const Model = GetLayoutViewModel();
+				if (!ParentComponent || !Model || Model->GetMode() == EDMXPixelMappingLayoutViewModelMode::LayoutNone)
+				{
+					return LOCTEXT("NoLayoutCanBeAppliedInfo", "Not applicable to the current selection");
+				}
+				else if (Model->GetMode() == EDMXPixelMappingLayoutViewModelMode::LayoutRendererComponentChildren)
+				{
+					return FText::Format(ChildrenOfText, ParentComponentNameText);
+				}
+				else if (Model->GetMode() == EDMXPixelMappingLayoutViewModelMode::LayoutFixtureGroupComponentChildren)
+				{
+					return FText::Format(ChildrenOfText, ParentComponentNameText);
+				}
+				else if (Model->GetMode() == EDMXPixelMappingLayoutViewModelMode::LayoutMatrixComponentChildren)
+				{
+					return FText::Format(ChildrenOfText, ParentComponentNameText);
+				}
+				return FText::GetEmpty();
+			}();
+
+		InfoTextBlock->SetText(InfoText);
+	}
 }
 
 UDMXPixelMappingLayoutViewModel* FDMXPixelMappingLayoutViewModelDetails::GetLayoutViewModel() const
diff --git a/Engine/Plugins/VirtualProduction/DMX/DMXPixelMapping/Source/DMXPixelMappingEditor/Private/Customizations/DMXPixelMappingLayoutViewModelDetails.h b/Engine/Plugins/VirtualProduction/DMX/DMXPixelMapping/Source/DMXPixelMappingEditor/Private/Customizations/DMXPixelMappingLayoutViewModelDetails.h
--- a/Engine/Plugins/VirtualProduction/DMX/DMXPixelMapping/Source/DMXPixelMappingEditor/Private/Customizations/DMXPixelMappingLayoutViewModelDetails.h
+++ b/Engine/Plugins/VirtualProduction/DMX/DMXPixelMapping/Source/DMXPixelMappingEditor/Private/Customizations/DMXPixelMappingLayoutViewModelDetails.h
@@ -3,11 +3,14 @@
 #pragma once
 
 #include "IDetailCustomization.h"
+#include "Input/Reply.h"
 
+struct EVisibility;
 class FDMXPixelMappingToolkit;
-class UDMXPixelMappingLayoutViewModel;
-
+class IPropertyHandle;
 class IPropertyUtilities;
+class STextBlock;
+class UDMXPixelMappingLayoutViewModel;
 
 
 /** Details customization for the Layout View Model */
@@ -27,7 +30,19 @@ public:
 	virtual void CustomizeDetails(IDetailLayoutBuilder& DetailLayout) override;
 	//~ End IDetailCustomization interface 
 
-protected:
+private:
+	/** Returns the visibility for the bAutoApply property */
+	EVisibility GetAutoApplyPropertyVisiblity() const;
+
+	/** Returns the visibility for the 'Apply Layout Script' button */
+	EVisibility GetApplyLayoutScriptButtonVisibility() const;
+
+	/** Called when the apply layout script button was clicked */
+	FReply OnApplyLayoutScriptClicked();
+
+	/** Updates the info text */
+	void UpdateInfoText();
+
 	/** Retuns the object that is being customized */
 	UDMXPixelMappingLayoutViewModel* GetLayoutViewModel() const;
 
@@ -40,6 +55,12 @@ protected:
 	/** If true, shows an info why the Layout Script Class property is hidden */
 	bool bShowLayoutScriptClassHiddenInfo = true;
 
+	/** Text block displaying info about what's being arranged */
+	TSharedPtr<STextBlock> InfoTextBlock;
+
+	/** Property handle for the bAutoApply property */
+	TSharedPtr<IPropertyHandle> AutoApplyHandle;
+
 	/** Property utilities for this customization */
 	TSharedPtr<IPropertyUtilities> PropertyUtilities;
 
diff --git a/Engine/Plugins/VirtualProduction/DMX/DMXPixelMapping/Source/DMXPixelMappingEditor/Private/DMXPixelMappingEditorCommands.cpp b/Engine/Plugins/VirtualProduction/DMX/DMXPixelMapping/Source/DMXPixelMappingEditor/Private/DMXPixelMappingEditorCommands.cpp
--- a/Engine/Plugins/VirtualProduction/DMX/DMXPixelMapping/Source/DMXPixelMappingEditor/Private/DMXPixelMappingEditorCommands.cpp
+++ b/Engine/Plugins/VirtualProduction/DMX/DMXPixelMapping/Source/DMXPixelMappingEditor/Private/DMXPixelMappingEditorCommands.cpp
@@ -41,7 +41,6 @@ void FDMXPixelMappingEditorCommands::RegisterCommands()
 
 	UI_COMMAND(ToggleScaleChildrenWithParent, "Scale Children with Parent", "Sets if children are scaled with their parent, when the parent is resized.", EUserInterfaceActionType::ToggleButton, FInputChord(EModifierKey::Alt, EKeys::Q));
 	UI_COMMAND(ToggleAlwaysSelectGroup, "Always select Group", "Sets if the parent is selected, if a child is clicked.", EUserInterfaceActionType::ToggleButton, FInputChord(EModifierKey::Alt, EKeys::W));
-	UI_COMMAND(ToggleApplyLayoutScriptWhenLoaded, "Apply Layout Script instantly", "Sets if layout script are applied as soon as they are loaded.", EUserInterfaceActionType::ToggleButton, FInputChord(EModifierKey::Alt, EKeys::E));
 	UI_COMMAND(ToggleShowMatrixCells, "Show Matrix Cells", "Sets if matrix cells are displayed. Can be turned off for better editor performance when pixel mapping large quantities of fixtures.", EUserInterfaceActionType::ToggleButton, FInputChord(EModifierKey::Alt, EKeys::A));
 	UI_COMMAND(ToggleShowComponentNames, "Show Component Names", "Sets if the name of components are displayed. Can be turned off for better editor performance when pixel mapping large quantities of fixtures.", EUserInterfaceActionType::ToggleButton, FInputChord(EModifierKey::Alt, EKeys::S));
 	UI_COMMAND(ToggleShowPatchInfo, "Show Patch Info", "Sets if information about the pach is displayed. Can be turned off for better editor performance when pixel mapping large quantities of fixtures.", EUserInterfaceActionType::ToggleButton, FInputChord(EModifierKey::Alt, EKeys::D));
diff --git a/Engine/Plugins/VirtualProduction/DMX/DMXPixelMapping/Source/DMXPixelMappingEditor/Private/DMXPixelMappingEditorCommands.h b/Engine/Plugins/VirtualProduction/DMX/DMXPixelMapping/Source/DMXPixelMappingEditor/Private/DMXPixelMappingEditorCommands.h
--- a/Engine/Plugins/VirtualProduction/DMX/DMXPixelMapping/Source/DMXPixelMappingEditor/Private/DMXPixelMappingEditorCommands.h
+++ b/Engine/Plugins/VirtualProduction/DMX/DMXPixelMapping/Source/DMXPixelMappingEditor/Private/DMXPixelMappingEditorCommands.h
@@ -39,7 +39,6 @@ public:
 	
 	TSharedPtr<FUICommandInfo> ToggleScaleChildrenWithParent;
 	TSharedPtr<FUICommandInfo> ToggleAlwaysSelectGroup;
-	TSharedPtr<FUICommandInfo> ToggleApplyLayoutScriptWhenLoaded;
 	TSharedPtr<FUICommandInfo> ToggleShowMatrixCells;
 	TSharedPtr<FUICommandInfo> ToggleShowComponentNames;
 	TSharedPtr<FUICommandInfo> ToggleShowPatchInfo;
diff --git a/Engine/Plugins/VirtualProduction/DMX/DMXPixelMapping/Source/DMXPixelMappingEditor/Private/Toolkits/DMXPixelMappingToolkit.cpp b/Engine/Plugins/VirtualProduction/DMX/DMXPixelMapping/Source/DMXPixelMappingEditor/Private/Toolkits/DMXPixelMappingToolkit.cpp
--- a/Engine/Plugins/VirtualProduction/DMX/DMXPixelMapping/Source/DMXPixelMappingEditor/Private/Toolkits/DMXPixelMappingToolkit.cpp
+++ b/Engine/Plugins/VirtualProduction/DMX/DMXPixelMapping/Source/DMXPixelMappingEditor/Private/Toolkits/DMXPixelMappingToolkit.cpp
@@ -1059,7 +1059,6 @@ void FDMXPixelMappingToolkit::SetupCommands()
 	UDMXPixelMappingEditorSettings* EditorSettings = GetMutableDefault<UDMXPixelMappingEditorSettings>(); 
 	UE_DMX_MAP_EDITOR_SETTING_TO_TOGGLE_COMMAND(EditorSettings, DesignerSettings, bScaleChildrenWithParent, ToggleScaleChildrenWithParent);
 	UE_DMX_MAP_EDITOR_SETTING_TO_TOGGLE_COMMAND(EditorSettings, DesignerSettings, bAlwaysSelectGroup, ToggleAlwaysSelectGroup);
-	UE_DMX_MAP_EDITOR_SETTING_TO_TOGGLE_COMMAND(EditorSettings, DesignerSettings, bApplyLayoutScriptWhenLoaded, ToggleApplyLayoutScriptWhenLoaded);
 	UE_DMX_MAP_EDITOR_SETTING_TO_TOGGLE_COMMAND(EditorSettings, DesignerSettings, bShowMatrixCells, ToggleShowMatrixCells);
 	UE_DMX_MAP_EDITOR_SETTING_TO_TOGGLE_COMMAND(EditorSettings, DesignerSettings, bShowComponentNames, ToggleShowComponentNames);
 	UE_DMX_MAP_EDITOR_SETTING_TO_TOGGLE_COMMAND(EditorSettings, DesignerSettings, bShowPatchInfo, ToggleShowPatchInfo);
diff --git a/Engine/Plugins/VirtualProduction/DMX/DMXPixelMapping/Source/DMXPixelMappingEditor/Private/ViewModels/DMXPixelMappingLayoutViewModel.cpp b/Engine/Plugins/VirtualProduction/DMX/DMXPixelMapping/Source/DMXPixelMappingEditor/Private/ViewModels/DMXPixelMappingLayoutViewModel.cpp
--- a/Engine/Plugins/VirtualProduction/DMX/DMXPixelMapping/Source/DMXPixelMappingEditor/Private/ViewModels/DMXPixelMappingLayoutViewModel.cpp
+++ b/Engine/Plugins/VirtualProduction/DMX/DMXPixelMapping/Source/DMXPixelMappingEditor/Private/ViewModels/DMXPixelMappingLayoutViewModel.cpp
@@ -17,22 +17,6 @@
 
 #define LOCTEXT_NAMESPACE "DMXPixelMappingLayoutViewModel"
 
-UDMXPixelMappingLayoutViewModel::UDMXPixelMappingLayoutViewModel()
-{
-	if (GEditor)
-	{
-		GEditor->RegisterForUndo(this);
-	}
-}
-
-UDMXPixelMappingLayoutViewModel::~UDMXPixelMappingLayoutViewModel()
-{
-	if (GEditor)
-	{
-		GEditor->UnregisterForUndo(this);
-	}
-}
-
 void UDMXPixelMappingLayoutViewModel::PreEditChange(FProperty* PropertyAboutToChange)
 {
 	Super::PreEditChange(PropertyAboutToChange);
@@ -65,7 +49,7 @@ void UDMXPixelMappingLayoutViewModel::SetToolkit(const TSharedRef<FDMXPixelMappi
 	if (DesignerSettings.bApplyLayoutScriptWhenLoaded)
 	{
 		FScopedTransaction ApplyLayoutScriptTransactionDirect(LOCTEXT("ApplyLayoutScriptTransactionDirect", "Apply Layout Script"));
-		ApplyLayoutScripts();
+		ForceApplyLayoutScript();
 	}
 
 	InToolkit->GetOnSelectedComponentsChangedDelegate().AddUObject(this, &UDMXPixelMappingLayoutViewModel::OnSelectedComponentsChanged);
@@ -96,15 +80,15 @@ TArray<UObject*> UDMXPixelMappingLayoutViewModel::GetLayoutScriptsObjectsSlow()
 {
 	TArray<UObject*> Result;
 
-	const EDMXPixelMappingLayoutViewModelMode LayoutType = GetMode();
-	if (LayoutType == EDMXPixelMappingLayoutViewModelMode::LayoutRendererComponentChildren)
+	const EDMXPixelMappingLayoutViewModelMode LayoutMode = GetMode();
+	if (LayoutMode == EDMXPixelMappingLayoutViewModelMode::LayoutRendererComponentChildren)
 	{
 		if (RendererComponent.IsValid())
 		{
 			Result.Add(RendererComponent->LayoutScript);
 		}
 	}
-	else if (LayoutType == EDMXPixelMappingLayoutViewModelMode::LayoutFixtureGroupComponentChildren)
+	else if (LayoutMode == EDMXPixelMappingLayoutViewModelMode::LayoutFixtureGroupComponentChildren)
 	{
 		for (TWeakObjectPtr<UDMXPixelMappingFixtureGroupComponent> FixtureGroup : FixtureGroupComponents)
 		{
@@ -114,7 +98,7 @@ TArray<UObject*> UDMXPixelMappingLayoutViewModel::GetLayoutScriptsObjectsSlow()
 			}
 		}
 	}
-	else if (LayoutType == EDMXPixelMappingLayoutViewModelMode::LayoutMatrixComponentChildren)
+	else if (LayoutMode == EDMXPixelMappingLayoutViewModelMode::LayoutMatrixComponentChildren)
 	{
 		for (TWeakObjectPtr<UDMXPixelMappingMatrixComponent> Matrix : MatrixComponents)
 		{
@@ -128,53 +112,100 @@ TArray<UObject*> UDMXPixelMappingLayoutViewModel::GetLayoutScriptsObjectsSlow()
 	return Result;
 }
 
-void UDMXPixelMappingLayoutViewModel::RequestApplyLayoutScripts()
+bool UDMXPixelMappingLayoutViewModel::CanApplyLayoutScript() const
 {
-	if (!ApplyLayoutScriptTimerHandle.IsValid())
-	{
-		ApplyLayoutScriptTimerHandle = GEditor->GetTimerManager()->SetTimerForNextTick(FTimerDelegate::CreateUObject(this, &UDMXPixelMappingLayoutViewModel::ApplyLayoutScripts));
-	}
+	const EDMXPixelMappingLayoutViewModelMode LayoutMode = GetMode();
+	return 
+		LayoutScriptClass.IsValid() &&
+		LayoutMode != EDMXPixelMappingLayoutViewModelMode::LayoutNone;
 }
 
-void UDMXPixelMappingLayoutViewModel::PostUndo(bool bSuccess)
+void UDMXPixelMappingLayoutViewModel::RequestApplyLayoutScript()
 {
-	RefreshComponents();
-	RefreshLayoutScriptClass();
-}
-
-void UDMXPixelMappingLayoutViewModel::PostRedo(bool bSuccess)
-{
-	RefreshComponents();
-	RefreshLayoutScriptClass();
+	if (!ApplyLayoutScriptTimerHandle.IsValid())
+	{
+		ApplyLayoutScriptTimerHandle = GEditor->GetTimerManager()->SetTimerForNextTick(FTimerDelegate::CreateUObject(this, &UDMXPixelMappingLayoutViewModel::ForceApplyLayoutScript));
+	}
 }
 
-void UDMXPixelMappingLayoutViewModel::ApplyLayoutScripts()
+void UDMXPixelMappingLayoutViewModel::ForceApplyLayoutScript()
 {
 	ApplyLayoutScriptTimerHandle.Invalidate();
 
-	const EDMXPixelMappingLayoutViewModelMode LayoutType = GetMode();
-	if (LayoutType == EDMXPixelMappingLayoutViewModelMode::LayoutRendererComponentChildren)
+	const EDMXPixelMappingLayoutViewModelMode LayoutMode = GetMode();
+	if (LayoutMode == EDMXPixelMappingLayoutViewModelMode::LayoutRendererComponentChildren)
 	{
 		LayoutRendererComponentChildren();
 	}
-	else if (LayoutType == EDMXPixelMappingLayoutViewModelMode::LayoutFixtureGroupComponentChildren)
+	else if (LayoutMode == EDMXPixelMappingLayoutViewModelMode::LayoutFixtureGroupComponentChildren)
 	{
 		LayoutFixtureGroupComponentChildren();
 	}
-	else if (LayoutType == EDMXPixelMappingLayoutViewModelMode::LayoutMatrixComponentChildren)
+	else if (LayoutMode == EDMXPixelMappingLayoutViewModelMode::LayoutMatrixComponentChildren)
 	{
 		LayoutMatrixComponentChildren();
 	}
 	else
 	{
-		ensureAlwaysMsgf(LayoutType == EDMXPixelMappingLayoutViewModelMode::LayoutNone, TEXT("Unhandled Layout Type, layout cannot be applied"));
+		ensureAlwaysMsgf(LayoutMode == EDMXPixelMappingLayoutViewModelMode::LayoutNone, TEXT("Unhandled Layout Type, layout cannot be applied"));
 	}
 }
 
+UDMXPixelMappingOutputComponent* UDMXPixelMappingLayoutViewModel::GetParentComponent() const
+{
+	// Some ensures here to make sure we rely on the same conditions as what we assume the specific mode means (should correspond to GetMode);
+	const EDMXPixelMappingLayoutViewModelMode LayoutMode = GetMode();
+	if (LayoutMode == EDMXPixelMappingLayoutViewModelMode::LayoutRendererComponentChildren)
+	{
+		if (!ensureMsgf(RendererComponent.IsValid(), TEXT("GetMode no longer matches assumed conditions.")))
+		{
+			return nullptr;
+		}
+
+		return RendererComponent.Get();
+	}
+	else if (LayoutMode == EDMXPixelMappingLayoutViewModelMode::LayoutFixtureGroupComponentChildren)
+	{
+		if (!ensureMsgf(FixtureGroupComponents.Num() == 1 && ScreenComponents.IsEmpty(), TEXT("GetMode no longer matches assumed conditions.")))
+		{
+			return nullptr;
+		}
+
+		return FixtureGroupComponents[0].Get();
+	}
+	else if (LayoutMode == EDMXPixelMappingLayoutViewModelMode::LayoutMatrixComponentChildren)
+	{
+		if (!ensureMsgf(MatrixComponents.Num() == 1 && FixtureGroupComponents.IsEmpty(), TEXT("GetMode no longer matches assumed conditions.")))
+		{
+			return nullptr;
+		}
+
+		return MatrixComponents[0].Get();
+	}
+
+	return nullptr;
+}
+
+void UDMXPixelMappingLayoutViewModel::PostUndo(bool bSuccess)
+{
+	RefreshComponents();
+	RefreshLayoutScriptClass();
+
+	OnModelChanged.Broadcast();
+}
+
+void UDMXPixelMappingLayoutViewModel::PostRedo(bool bSuccess)
+{
+	RefreshComponents();
+	RefreshLayoutScriptClass();
+
+	OnModelChanged.Broadcast();
+}
+
 void UDMXPixelMappingLayoutViewModel::LayoutRendererComponentChildren()
 {
-	const EDMXPixelMappingLayoutViewModelMode LayoutType = GetMode();
-	if (!ensureMsgf(LayoutType == EDMXPixelMappingLayoutViewModelMode::LayoutRendererComponentChildren, TEXT("Trying to layout children of Renderer Components, but the current selection does not support this.")))
+	const EDMXPixelMappingLayoutViewModelMode LayoutMode = GetMode();
+	if (!ensureMsgf(LayoutMode == EDMXPixelMappingLayoutViewModelMode::LayoutRendererComponentChildren, TEXT("Trying to layout children of Renderer Components, but the current selection does not support this.")))
 	{
 		return;
 	}
@@ -205,8 +236,8 @@ void UDMXPixelMappingLayoutViewModel::LayoutRendererComponentChildren()
 
 void UDMXPixelMappingLayoutViewModel::LayoutFixtureGroupComponentChildren()
 {
-	const EDMXPixelMappingLayoutViewModelMode LayoutType = GetMode();
-	if (!ensureMsgf(!FixtureGroupComponents.IsEmpty() && LayoutType == EDMXPixelMappingLayoutViewModelMode::LayoutFixtureGroupComponentChildren, TEXT("Trying to layout children of Fixture Group Components, but the current selection does not support this.")))
+	const EDMXPixelMappingLayoutViewModelMode LayoutMode = GetMode();
+	if (!ensureMsgf(!FixtureGroupComponents.IsEmpty() && LayoutMode == EDMXPixelMappingLayoutViewModelMode::LayoutFixtureGroupComponentChildren, TEXT("Trying to layout children of Fixture Group Components, but the current selection does not support this.")))
 	{
 		return;
 	}
@@ -220,6 +251,10 @@ void UDMXPixelMappingLayoutViewModel::LayoutFixtureGroupComponentChildren()
 				continue;
 			}
 
+			// Tokens expect unrotated space
+			const double RestoreRotation = FixtureGroupComponent->GetRotation();
+			FixtureGroupComponent->SetRotation(0.0);
+
 			InitializeLayoutScript(FixtureGroupComponent, FixtureGroupComponent->LayoutScript);
 
 			TArray<FDMXPixelMappingLayoutToken> InLayoutTokens;
@@ -237,14 +272,16 @@ void UDMXPixelMappingLayoutViewModel::LayoutFixtureGroupComponentChildren()
 			FixtureGroupComponent->PreEditChange(UDMXPixelMappingFixtureGroupComponent::StaticClass()->FindPropertyByName(GET_MEMBER_NAME_CHECKED(UDMXPixelMappingFixtureGroupComponent, LayoutScript)));
 			ApplyLayoutTokens(OutLayoutTokens);
 			FixtureGroupComponent->PostEditChange();
+
+			FixtureGroupComponent->SetRotation(RestoreRotation);
 		}
 	}
 }
 
 void UDMXPixelMappingLayoutViewModel::LayoutMatrixComponentChildren()
 {
-	const EDMXPixelMappingLayoutViewModelMode LayoutType = GetMode();
-	if (!ensureMsgf(!MatrixComponents.IsEmpty() && LayoutType == EDMXPixelMappingLayoutViewModelMode::LayoutMatrixComponentChildren, TEXT("Trying to layout children of Fixture Group Components, but the current selection does not support this.")))
+	const EDMXPixelMappingLayoutViewModelMode LayoutMode = GetMode();
+	if (!ensureMsgf(!MatrixComponents.IsEmpty() && LayoutMode == EDMXPixelMappingLayoutViewModelMode::LayoutMatrixComponentChildren, TEXT("Trying to layout children of Fixture Group Components, but the current selection does not support this.")))
 	{
 		return;
 	}
@@ -258,6 +295,10 @@ void UDMXPixelMappingLayoutViewModel::LayoutMatrixComponentChildren()
 				continue;
 			}
 
+			// Tokens expect unrotated space
+			const double RestoreRotation = MatrixComponent->GetRotation();
+			MatrixComponent->SetRotation(0.0);
+
 			InitializeLayoutScript(MatrixComponent, MatrixComponent->LayoutScript);
 
 			TArray<FDMXPixelMappingLayoutToken> InLayoutTokens;
@@ -275,31 +316,33 @@ void UDMXPixelMappingLayoutViewModel::LayoutMatrixComponentChildren()
 			MatrixComponent->PreEditChange(UDMXPixelMappingMatrixComponent::StaticClass()->FindPropertyByName(GET_MEMBER_NAME_CHECKED(UDMXPixelMappingMatrixComponent, LayoutScript)));
 			ApplyLayoutTokens(OutLayoutTokens);
 			MatrixComponent->PostEditChange();
+
+			MatrixComponent->SetRotation(RestoreRotation);
 		}
 	}
 }
 
 void UDMXPixelMappingLayoutViewModel::ApplyLayoutTokens(const TArray<FDMXPixelMappingLayoutToken>& LayoutTokens) const
 {
-
 	for (const FDMXPixelMappingLayoutToken& LayoutToken : LayoutTokens)
 	{
 		if (UDMXPixelMappingOutputComponent* OutputComponent = LayoutToken.Component.Get())
 		{
-			if (LayoutToken.SizeX < 1.f || LayoutToken.SizeY < 1.f)
+			FVector2D Size(LayoutToken.SizeX, LayoutToken.SizeY);
+			if (LayoutToken.SizeX <= 0.f || LayoutToken.SizeY <= 0.f)
 			{
-				UE_LOG(LogDMXPixelMappingEditor, Warning, TEXT("Trying to apply layout script. But layout defines a component size < 1. This is not supported."));
-			}
-
-			if (OutputComponent->GetParent())
-			{
-				OutputComponent->GetParent()->Modify();
+				UE_LOG(LogDMXPixelMappingEditor, Warning, TEXT("Trying to apply layout script. But layout defines a component size <= 0. This is not supported. Setting a small size instead."));
+				Size.X = FMath::Max(UE_SMALL_NUMBER, LayoutToken.SizeX);
+				Size.Y = FMath::Max(UE_SMALL_NUMBER, LayoutToken.SizeY);
 			}
 
 			OutputComponent->Modify();
 
 			OutputComponent->SetPosition(FVector2D(LayoutToken.PositionX, LayoutToken.PositionY));
-			OutputComponent->SetSize(FVector2D(LayoutToken.SizeX, LayoutToken.SizeY));
+			OutputComponent->SetSize(Size);
+
+			// Parent rotation was set to 0 before initialized. Hence the relative rotation is the current rotation.
+			OutputComponent->SetRotation(LayoutToken.RelativeRotation);
 		}
 	}
 }
@@ -317,7 +360,7 @@ void UDMXPixelMappingLayoutViewModel::OnLayoutScriptClassChanged()
 	const FDMXPixelMappingDesignerSettings& DesignerSettings = GetDefault<UDMXPixelMappingEditorSettings>()->DesignerSettings;
 	if (DesignerSettings.bApplyLayoutScriptWhenLoaded && LayoutScriptClass)
 	{
-		ApplyLayoutScripts();
+		ForceApplyLayoutScript();
 	}
 }
 
@@ -338,8 +381,8 @@ void UDMXPixelMappingLayoutViewModel::InstantiateLayoutScripts()
 	}
 
 	UClass* StrongLayoutScriptClass = LayoutScriptClass.Get();
-	const EDMXPixelMappingLayoutViewModelMode LayoutType = GetMode();
-	if (LayoutType == EDMXPixelMappingLayoutViewModelMode::LayoutRendererComponentChildren)
+	const EDMXPixelMappingLayoutViewModelMode LayoutMode = GetMode();
+	if (LayoutMode == EDMXPixelMappingLayoutViewModelMode::LayoutRendererComponentChildren)
 	{
 		if (RendererComponent.IsValid() && 
 		   (!RendererComponent->LayoutScript ||	RendererComponent->LayoutScript->GetClass() != StrongLayoutScriptClass))
@@ -356,7 +399,7 @@ void UDMXPixelMappingLayoutViewModel::InstantiateLayoutScripts()
 			}
 		}
 	}
-	else if (LayoutType == EDMXPixelMappingLayoutViewModelMode::LayoutFixtureGroupComponentChildren)
+	else if (LayoutMode == EDMXPixelMappingLayoutViewModelMode::LayoutFixtureGroupComponentChildren)
 	{
 		for (TWeakObjectPtr<UDMXPixelMappingFixtureGroupComponent> FixtureGroupComponent : FixtureGroupComponents)
 		{
@@ -376,7 +419,7 @@ void UDMXPixelMappingLayoutViewModel::InstantiateLayoutScripts()
 			}
 		}
 	}
-	else if (LayoutType == EDMXPixelMappingLayoutViewModelMode::LayoutMatrixComponentChildren)
+	else if (LayoutMode == EDMXPixelMappingLayoutViewModelMode::LayoutMatrixComponentChildren)
 	{
 		for (TWeakObjectPtr<UDMXPixelMappingMatrixComponent> MatrixComponent : MatrixComponents)
 		{
@@ -398,7 +441,7 @@ void UDMXPixelMappingLayoutViewModel::InstantiateLayoutScripts()
 	}
 	else
 	{
-		ensureMsgf(LayoutType == EDMXPixelMappingLayoutViewModelMode::LayoutNone, TEXT("Unhandled Layout Type, layout cannot be applied"));
+		ensureMsgf(LayoutMode == EDMXPixelMappingLayoutViewModelMode::LayoutNone, TEXT("Unhandled Layout Type, layout cannot be applied"));
 	}
 }
 
@@ -412,6 +455,7 @@ void UDMXPixelMappingLayoutViewModel::InitializeLayoutScript(UDMXPixelMappingOut
 	LayoutScript->SetNumTokens(OutputComponent->GetChildren().Num());
 	LayoutScript->SetParentComponentPosition(OutputComponent->GetPosition());
 	LayoutScript->SetParentComponentSize(OutputComponent->GetSize());
+	LayoutScript->SetParentComponentRotation(OutputComponent->GetRotation());
 	LayoutScript->SetTextureSize(OutputComponent->GetRendererComponent()->GetSize());
 }
 
@@ -478,15 +522,15 @@ void UDMXPixelMappingLayoutViewModel::RefreshLayoutScriptClass()
 {
 	LayoutScriptClass = nullptr;
 
-	const EDMXPixelMappingLayoutViewModelMode LayoutType = GetMode();
-	if (LayoutType == EDMXPixelMappingLayoutViewModelMode::LayoutRendererComponentChildren)
+	const EDMXPixelMappingLayoutViewModelMode LayoutMode = GetMode();
+	if (LayoutMode == EDMXPixelMappingLayoutViewModelMode::LayoutRendererComponentChildren)
 	{
 		if (RendererComponent.IsValid() && RendererComponent->LayoutScript)
 		{
 			LayoutScriptClass =RendererComponent->LayoutScript->GetClass();
 		}
 	}
-	else if (LayoutType == EDMXPixelMappingLayoutViewModelMode::LayoutFixtureGroupComponentChildren)
+	else if (LayoutMode == EDMXPixelMappingLayoutViewModelMode::LayoutFixtureGroupComponentChildren)
 	{
 		for (TWeakObjectPtr<UDMXPixelMappingFixtureGroupComponent> FixtureGroup : FixtureGroupComponents)
 		{
@@ -507,7 +551,7 @@ void UDMXPixelMappingLayoutViewModel::RefreshLayoutScriptClass()
 			}
 		}
 	}
-	else if (LayoutType == EDMXPixelMappingLayoutViewModelMode::LayoutMatrixComponentChildren)
+	else if (LayoutMode == EDMXPixelMappingLayoutViewModelMode::LayoutMatrixComponentChildren)
 	{
 		for (TWeakObjectPtr<UDMXPixelMappingMatrixComponent> Matrix : MatrixComponents)
 		{
@@ -530,7 +574,7 @@ void UDMXPixelMappingLayoutViewModel::RefreshLayoutScriptClass()
 	}
 	else
 	{
-		ensureMsgf(LayoutType == EDMXPixelMappingLayoutViewModelMode::LayoutNone, TEXT("Unhandled Layout Type, layout cannot be applied"));
+		ensureMsgf(LayoutMode == EDMXPixelMappingLayoutViewModelMode::LayoutNone, TEXT("Unhandled Layout Type, layout cannot be applied"));
 	}
 }
 
diff --git a/Engine/Plugins/VirtualProduction/DMX/DMXPixelMapping/Source/DMXPixelMappingEditor/Private/ViewModels/DMXPixelMappingLayoutViewModel.h b/Engine/Plugins/VirtualProduction/DMX/DMXPixelMapping/Source/DMXPixelMappingEditor/Private/ViewModels/DMXPixelMappingLayoutViewModel.h
--- a/Engine/Plugins/VirtualProduction/DMX/DMXPixelMapping/Source/DMXPixelMappingEditor/Private/ViewModels/DMXPixelMappingLayoutViewModel.h
+++ b/Engine/Plugins/VirtualProduction/DMX/DMXPixelMapping/Source/DMXPixelMappingEditor/Private/ViewModels/DMXPixelMappingLayoutViewModel.h
@@ -31,7 +31,7 @@ enum class EDMXPixelMappingLayoutViewModelMode
 UCLASS(Transient)
 class UDMXPixelMappingLayoutViewModel
 	: public UObject
-	, public FEditorUndoClient
+	, public FSelfRegisteringEditorUndoClient
 {
 	GENERATED_BODY()
 
@@ -42,12 +42,6 @@ protected:
 #endif
 
 public:
-	/** Destructor */
-	UDMXPixelMappingLayoutViewModel();
-
-	/** Destructor */
-	virtual ~UDMXPixelMappingLayoutViewModel();
-
 	/** Sets the toolkit from which the model sources */
 	void SetToolkit(const TSharedRef<FDMXPixelMappingToolkit>& InToolkit);
 
@@ -57,13 +51,26 @@ public:
 	/** Returns the Layout Scripts of relevant components. Looks up each relevant component and returns its script. */
 	TArray<UObject*> GetLayoutScriptsObjectsSlow() const;
 
-	/** Applies the layouts scripts currently in use on the next tick */
-	void RequestApplyLayoutScripts();
+	/** Returns true if a layout script can be applied. */
+	bool CanApplyLayoutScript() const;
+
+	/** Applies the layouts script on the next tick. Test if CanApplyLayoutScript before calling (ensured). */
+	void RequestApplyLayoutScript();
+
+	/** Applies the layouts script. Test if CanApplyLayoutScript before calling (ensured). */
+	void ForceApplyLayoutScript();
+
+	/** Returns true if the layout script should be auto-applied on property changes */
+	bool ShouldAutoApplyLayoutScript() const { return bAutoApply; }
+
+	/** Returns the parent of the components being laid out */
+	UDMXPixelMappingOutputComponent* GetParentComponent() const;
 
 	/** Delegate executed when the Model changed */
 	FSimpleMulticastDelegate OnModelChanged;
 
 	// Property name getters
+	FORCEINLINE static FName GetAutoApplyPropertyName() { return GET_MEMBER_NAME_CHECKED(UDMXPixelMappingLayoutViewModel, bAutoApply); }
 	FORCEINLINE static FName GetLayoutScriptClassPropertyName() { return GET_MEMBER_NAME_CHECKED(UDMXPixelMappingLayoutViewModel, LayoutScriptClass); }
 
 protected:
@@ -73,9 +80,6 @@ protected:
 	// End of FEditorUndoClient
 
 private:
-	/** Applies the layouts scripts currently in use */
-	void ApplyLayoutScripts();
-
 	/** Lays out children of the Renderer Component */
 	void LayoutRendererComponentChildren();
 
@@ -125,7 +129,11 @@ private:
 	UPROPERTY()
 	TArray<TWeakObjectPtr<UDMXPixelMappingMatrixComponent>> MatrixComponents;
 
-	/** The Layout Script class currently in use */
+	/** If true, the script is auto applied on property changes */
+	UPROPERTY(EditAnywhere, Category = "Layout Script", Meta = (DisplayName = "Apply Script on Property Changes"))
+	bool bAutoApply = true;
+
+	/** The Layout Script class currently in use. Note, custom scripts can be blueprinted. */
 	UPROPERTY(EditAnywhere, Category = "Layout Script", Meta = (ShowDisplayNames))
 	TSoftClassPtr<UDMXPixelMappingLayoutScript> LayoutScriptClass;
 
@@ -133,5 +141,6 @@ private:
 	FTimerHandle ApplyLayoutScriptTimerHandle;
 
 	/** The Layout Script class currently before it was changed */
+	UPROPERTY(Transient)
 	TSoftClassPtr<UDMXPixelMappingLayoutScript> PreEditChangeLayoutScriptClass;
 };
diff --git a/Engine/Plugins/VirtualProduction/DMX/DMXPixelMapping/Source/DMXPixelMappingEditor/Private/Views/SDMXPixelMappingLayoutView.cpp b/Engine/Plugins/VirtualProduction/DMX/DMXPixelMapping/Source/DMXPixelMappingEditor/Private/Views/SDMXPixelMappingLayoutView.cpp
--- a/Engine/Plugins/VirtualProduction/DMX/DMXPixelMapping/Source/DMXPixelMappingEditor/Private/Views/SDMXPixelMappingLayoutView.cpp
+++ b/Engine/Plugins/VirtualProduction/DMX/DMXPixelMapping/Source/DMXPixelMappingEditor/Private/Views/SDMXPixelMappingLayoutView.cpp
@@ -88,7 +88,10 @@ void SDMXPixelMappingLayoutView::PostRedo(bool bSuccess)
 
 void SDMXPixelMappingLayoutView::NotifyPostChange(const FPropertyChangedEvent& PropertyChangedEvent, FProperty* PropertyThatChanged)
 {
-	Model->RequestApplyLayoutScripts();
+	if (Model->ShouldAutoApplyLayoutScript())
+	{
+		Model->RequestApplyLayoutScript();
+	}
 }
 
 void SDMXPixelMappingLayoutView::Refresh()
diff --git a/Engine/Plugins/VirtualProduction/DMX/DMXPixelMapping/Source/DMXPixelMappingEditor/Private/Widgets/SDMXPixelMappingDesignerToolbar.cpp b/Engine/Plugins/VirtualProduction/DMX/DMXPixelMapping/Source/DMXPixelMappingEditor/Private/Widgets/SDMXPixelMappingDesignerToolbar.cpp
--- a/Engine/Plugins/VirtualProduction/DMX/DMXPixelMapping/Source/DMXPixelMappingEditor/Private/Widgets/SDMXPixelMappingDesignerToolbar.cpp
+++ b/Engine/Plugins/VirtualProduction/DMX/DMXPixelMapping/Source/DMXPixelMappingEditor/Private/Widgets/SDMXPixelMappingDesignerToolbar.cpp
@@ -297,11 +297,6 @@ namespace UE::DMX
 		}
 		MenuBuilder.EndSection();
 
-		MenuBuilder.BeginSection("LayoutSettings", LOCTEXT("LayoutSettingsSection", "Layout Settings"));
-		{
-			MenuBuilder.AddMenuEntry(FDMXPixelMappingEditorCommands::Get().ToggleApplyLayoutScriptWhenLoaded);
-		}
-
 		return MenuBuilder.MakeWidget();
 	}
 
diff --git a/Engine/Plugins/VirtualProduction/DMX/DMXPixelMapping/Source/DMXPixelMappingRuntime/Private/LayoutScripts/DMXPixelMappingLayoutScript.cpp b/Engine/Plugins/VirtualProduction/DMX/DMXPixelMapping/Source/DMXPixelMappingRuntime/Private/LayoutScripts/DMXPixelMappingLayoutScript.cpp
--- a/Engine/Plugins/VirtualProduction/DMX/DMXPixelMapping/Source/DMXPixelMappingRuntime/Private/LayoutScripts/DMXPixelMappingLayoutScript.cpp
+++ b/Engine/Plugins/VirtualProduction/DMX/DMXPixelMapping/Source/DMXPixelMappingRuntime/Private/LayoutScripts/DMXPixelMappingLayoutScript.cpp
@@ -24,6 +24,9 @@ FDMXPixelMappingLayoutToken::FDMXPixelMappingLayoutToken(TWeakObjectPtr<UDMXPixe
 	SizeX = OutputComponent->GetSize().X;
 	SizeY = OutputComponent->GetSize().Y;
 
+	// Layout tokens are created with an unrotated parent. Hence the relative rotation is the current rotation.
+	RelativeRotation = OutputComponent->GetRotation();
+
 	InitializeFixtureID();
 }
 
diff --git a/Engine/Plugins/VirtualProduction/DMX/DMXPixelMapping/Source/DMXPixelMappingRuntime/Private/LayoutScripts/DMXPixelMappingLayoutScript_GridLayout.cpp b/Engine/Plugins/VirtualProduction/DMX/DMXPixelMapping/Source/DMXPixelMappingRuntime/Private/LayoutScripts/DMXPixelMappingLayoutScript_GridLayout.cpp
--- a/Engine/Plugins/VirtualProduction/DMX/DMXPixelMapping/Source/DMXPixelMappingRuntime/Private/LayoutScripts/DMXPixelMappingLayoutScript_GridLayout.cpp
+++ b/Engine/Plugins/VirtualProduction/DMX/DMXPixelMapping/Source/DMXPixelMappingRuntime/Private/LayoutScripts/DMXPixelMappingLayoutScript_GridLayout.cpp
@@ -53,7 +53,7 @@ void UDMXPixelMappingLayoutScript_GridLayout::Layout_Implementation(const TArray
 	// No empty rows:
 	// Example, user sets the grid to 1x9 interactively, but there's 10 cells (tokens).
 	// Hence UI flips to 2x9, to respect his 9 value, and compensate for the missing cell.
-	// Now the Grid is 2x9 = 18 cells, but we have only 10 cells. So 4 rows remain empty.
+	// Now the Grid is 2x9 = 18 cells, but we have only 10 cells. The last 4 columns remain empty.
 	const int32 ActualRows = FMath::RoundFromZero((float)NumTokens / Columns);
 
 	// Cell Size
@@ -113,8 +113,9 @@ void UDMXPixelMappingLayoutScript_GridLayout::Layout_Implementation(const TArray
 
 			OutTokens[PositionIDIndex].PositionX = AbsolutePosition.X + Column * ComponentSizeX + ClampedPadding / 2.f;
 			OutTokens[PositionIDIndex].PositionY = AbsolutePosition.Y + Row * ComponentSizeY + ClampedPadding / 2.f;
-			OutTokens[PositionIDIndex].SizeX = FMath::Max(1.f, ComponentSizeX - ClampedPadding);
-			OutTokens[PositionIDIndex].SizeY = FMath::Max(1.f, ComponentSizeY - ClampedPadding);
+			OutTokens[PositionIDIndex].SizeX = FMath::Max(UE_SMALL_NUMBER, ComponentSizeX - ClampedPadding);
+			OutTokens[PositionIDIndex].SizeY = FMath::Max(UE_SMALL_NUMBER, ComponentSizeY - ClampedPadding);
+			OutTokens[PositionIDIndex].RelativeRotation = ParentRotation;
 		}
 		else
 		{
diff --git a/Engine/Plugins/VirtualProduction/DMX/DMXPixelMapping/Source/DMXPixelMappingRuntime/Private/LayoutScripts/DMXPixelMappingLayoutScript_GridLayout.h b/Engine/Plugins/VirtualProduction/DMX/DMXPixelMapping/Source/DMXPixelMappingRuntime/Private/LayoutScripts/DMXPixelMappingLayoutScript_GridLayout.h
--- a/Engine/Plugins/VirtualProduction/DMX/DMXPixelMapping/Source/DMXPixelMappingRuntime/Private/LayoutScripts/DMXPixelMappingLayoutScript_GridLayout.h
+++ b/Engine/Plugins/VirtualProduction/DMX/DMXPixelMapping/Source/DMXPixelMappingRuntime/Private/LayoutScripts/DMXPixelMappingLayoutScript_GridLayout.h
@@ -33,11 +33,11 @@ public:
 	virtual void SetNumTokens(int32 NewNumTokens) override;
 	//~ End DMXPixelMappingScript interface
 
-	/** Num Columns of the grid */
+	/** Num Columns in the grid */
 	UPROPERTY(EditAnywhere, Category = "Layout Settings")
 	int32 Columns = 1;
 
-	/** Num Rows of the grid */
+	/** Num Rows in the grid */
 	UPROPERTY(EditAnywhere, Category = "Layout Settings")
 	int32 Rows = 1;
 
@@ -46,11 +46,11 @@ public:
 	float Padding = 0.f;
 
 	/** Horizontal text alignment */
-	UPROPERTY(EditAnywhere, Category = "Layout Settings", Meta = (DisplayName = "Horizontal Alignment"))
+	UPROPERTY(EditAnywhere, Category = "Layout Settings")
 	TEnumAsByte<EHorizontalAlignment> HorizontalAlignment;
 
 	/** Vertical text alignment */
-	UPROPERTY(EditAnywhere, Category = "Layout Settings", Meta = (DisplayName = "Vertical Alignment"))
+	UPROPERTY(EditAnywhere, Category = "Layout Settings")
 	TEnumAsByte<EVerticalAlignment> VerticalAlignment;
 
 	/** How cells are distributed, compared to their previous distribution */
@@ -63,4 +63,8 @@ private:
 
 	/** Clamps colums to rows given num tokens */
 	void ClampRows();
+
+	/** Auto-size is enabled when either horizontal or vertical aligmenet are set to fill */
+	UPROPERTY()
+	bool bAutoSizeMode = true;
 };
diff --git a/Engine/Plugins/VirtualProduction/DMX/DMXPixelMapping/Source/DMXPixelMappingRuntime/Public/LayoutScripts/DMXPixelMappingLayoutScript.h b/Engine/Plugins/VirtualProduction/DMX/DMXPixelMapping/Source/DMXPixelMappingRuntime/Public/LayoutScripts/DMXPixelMappingLayoutScript.h
--- a/Engine/Plugins/VirtualProduction/DMX/DMXPixelMapping/Source/DMXPixelMappingRuntime/Public/LayoutScripts/DMXPixelMappingLayoutScript.h
+++ b/Engine/Plugins/VirtualProduction/DMX/DMXPixelMapping/Source/DMXPixelMappingRuntime/Public/LayoutScripts/DMXPixelMappingLayoutScript.h
@@ -13,7 +13,7 @@ class UDMXPixelMappingOutputComponent;
 class UTextureRenderTarget2D;
 
 
-/** The position and size of a comoponent in a layout */
+/** The position and size of a comoponent in a layout. */
 USTRUCT(BlueprintType)
 struct DMXPIXELMAPPINGRUNTIME_API FDMXPixelMappingLayoutToken
 {
@@ -38,6 +38,10 @@ struct DMXPIXELMAPPINGRUNTIME_API FDMXPixelMappingLayoutToken
 	UPROPERTY(BlueprintReadWrite, EditAnywhere, Category = "Layout Script")
 	float SizeY = 0.f;
 
+	/** The rotation relative to the parent the component, applied after position */
+	UPROPERTY(BlueprintReadWrite, EditAnywhere, Category = "Layout Script")
+	double RelativeRotation = 0.0;
+
 	/** The output component to which the token is applied */
 	UPROPERTY(BlueprintReadWrite, Category = "Pixel Mapping Layout")
 	TWeakObjectPtr<UDMXPixelMappingOutputComponent> Component;
@@ -55,11 +59,12 @@ private:
 /** 
  * Allows scripting of Pixel Mapping Component Layouts.
  * Override the Layout function to implement a layout. 
+ * 
+ * Note, layout scripts are applied in parent's unrotated space. 
  */
 UCLASS(BlueprintType, Blueprintable, EditInlineNew, Abstract)
 class DMXPIXELMAPPINGRUNTIME_API UDMXPixelMappingLayoutScript
 	: public UObject
-
 {
 	GENERATED_BODY()
 
@@ -73,7 +78,7 @@ public:
 	 */
 	UFUNCTION(BlueprintNativeEvent, Category = "DMX")
 	void Layout(const TArray<FDMXPixelMappingLayoutToken>& InTokens, TArray<FDMXPixelMappingLayoutToken>& OutTokens);
-	virtual void Layout_Implementation(const TArray<FDMXPixelMappingLayoutToken>& InTokens, TArray<FDMXPixelMappingLayoutToken>& OutTokens) PURE_VIRTUAL(UDMXPixelMappingLayoutScript::Layout_Implementation, return;)
+	virtual void Layout_Implementation(const TArray<FDMXPixelMappingLayoutToken>& InTokens, TArray<FDMXPixelMappingLayoutToken>& OutTokens) {}
 
 	/** Sets the number of Tokens that will be passed when Layout is called */
 	virtual void SetNumTokens(int32 NewNumTokens) { NumTokens = NewNumTokens; }
@@ -84,13 +89,16 @@ public:
 	/** Sets the size of the parent component where the components that are being layouted reside in. */
 	virtual void SetParentComponentSize(const FVector2D& NewSize) { ParentComponentSize = NewSize; };
 
+	/** Sets the size of the parent component where the components that are being layouted reside in. */
+	virtual void SetParentComponentRotation(double NewRotation) { ParentRotation = NewRotation; };
+
 	/** Sets the Texture Size */
 	virtual void SetTextureSize(const FVector2D& NewTextureSize) { TextureSize = NewTextureSize; }
 
 protected:
 	/** The number of tokens in the Layout. */
 	UPROPERTY(BlueprintReadOnly, Transient, Category = "Layout Script")
-	int32 NumTokens;
+	int32 NumTokens = 0;
 
 	/** The position of the parent component where the components that are being layouted reside in. */
 	UPROPERTY(BlueprintReadOnly, Transient, Category = "Layout Script")
@@ -100,7 +108,14 @@ protected:
 	UPROPERTY(BlueprintReadOnly, Transient, Category = "Layout Script")
 	FVector2D ParentComponentSize;
 
-	/** The size of the texture in the Pixel Mapping asset. */
+	/** 
+	 * The rotation of the parent component where the components that are being layouted reside in.
+	 * Note, the position property does not consider this rotation to ease calculations.
+	 */
+	UPROPERTY(BlueprintReadOnly, Transient, Category = "Layout Script")
+	double ParentRotation;
+
+	/** The size of the texture this component is mapping. */
 	UPROPERTY(BlueprintReadOnly, Transient, Category = "Layout Script")
 	FVector2D TextureSize;
 };
