DMX - PixelMapping now uses a cardinal cross as its mouse cursor when in rotate mode

diff --git a/Engine/Plugins/VirtualProduction/DMX/DMXPixelMapping/Source/DMXPixelMappingEditor/Private/Widgets/SDMXPixelMappingTransformHandle.cpp b/Engine/Plugins/VirtualProduction/DMX/DMXPixelMapping/Source/DMXPixelMappingEditor/Private/Widgets/SDMXPixelMappingTransformHandle.cpp
--- a/Engine/Plugins/VirtualProduction/DMX/DMXPixelMapping/Source/DMXPixelMappingEditor/Private/Widgets/SDMXPixelMappingTransformHandle.cpp
+++ b/Engine/Plugins/VirtualProduction/DMX/DMXPixelMapping/Source/DMXPixelMappingEditor/Private/Widgets/SDMXPixelMappingTransformHandle.cpp
@@ -188,14 +188,31 @@ FReply SDMXPixelMappingTransformHandle::OnMouseMove(const FGeometry& MyGeometry,
 
 FCursorReply SDMXPixelMappingTransformHandle::OnCursorQuery(const FGeometry& MyGeometry, const FPointerEvent& MouseEvent) const
 {
-	switch (TransformDirection)
+	const TSharedPtr<FDMXPixelMappingToolkit> Toolkit = DesignerViewWeakPtr.Pin()->GetToolkit();
+	if (!Toolkit.IsValid())
 	{
-	case EDMXPixelMappingTransformDirection::BottomRight:
-		return FCursorReply::Cursor(EMouseCursor::ResizeSouthEast);
-	case EDMXPixelMappingTransformDirection::BottomCenter:
-		return FCursorReply::Cursor(EMouseCursor::ResizeUpDown);
-	case EDMXPixelMappingTransformDirection::CenterRight:
-		return FCursorReply::Cursor(EMouseCursor::ResizeLeftRight);
+		return FCursorReply::Unhandled();
+	}
+
+	if (Toolkit->GetTransformHandleMode() == UE::DMX::EDMXPixelMappingTransformHandleMode::Resize)
+	{
+		switch (TransformDirection)
+		{
+		case EDMXPixelMappingTransformDirection::BottomRight:
+			return FCursorReply::Cursor(EMouseCursor::ResizeSouthEast);
+		case EDMXPixelMappingTransformDirection::BottomCenter:
+			return FCursorReply::Cursor(EMouseCursor::ResizeUpDown);
+		case EDMXPixelMappingTransformDirection::CenterRight:
+			return FCursorReply::Cursor(EMouseCursor::ResizeLeftRight);
+		}
+	}
+	else if (Toolkit->GetTransformHandleMode() == UE::DMX::EDMXPixelMappingTransformHandleMode::Rotate)
+	{
+		return FCursorReply::Cursor(EMouseCursor::CardinalCross);
+	}
+	else
+	{
+		checkf(0, TEXT("Unhandled enum value"));
 	}
 
 	return FCursorReply::Unhandled();
