DMX: Change the default quality of pixel mapping cells from low to high

diff --git a/Engine/Plugins/VirtualProduction/DMX/DMXPixelMapping/Source/DMXPixelMappingRuntime/Public/Components/DMXPixelMappingOutputDMXComponent.h b/Engine/Plugins/VirtualProduction/DMX/DMXPixelMapping/Source/DMXPixelMappingRuntime/Public/Components/DMXPixelMappingOutputDMXComponent.h
--- a/Engine/Plugins/VirtualProduction/DMX/DMXPixelMapping/Source/DMXPixelMappingRuntime/Public/Components/DMXPixelMappingOutputDMXComponent.h
+++ b/Engine/Plugins/VirtualProduction/DMX/DMXPixelMapping/Source/DMXPixelMappingRuntime/Public/Components/DMXPixelMappingOutputDMXComponent.h
@@ -3,6 +3,7 @@
 #pragma once
 
 #include "DMXPixelMappingOutputComponent.h"
+#include "DMXPixelMappingRenderElement.h"
 #include "Library/DMXEntityReference.h"
 
 #include "DMXPixelMappingOutputDMXComponent.generated.h"
@@ -32,7 +33,7 @@ public:
 
 	/** The quality level to use when averaging colors during downsampling. */
 	UPROPERTY(EditAnywhere, BlueprintReadOnly, Category = "Quality")
-	EDMXPixelBlendingQuality CellBlendingQuality;
+	EDMXPixelBlendingQuality CellBlendingQuality = EDMXPixelBlendingQuality::High;
 
 	UPROPERTY(VisibleAnywhere, BlueprintReadOnly, Category = "Fixture Patch", Meta = (ShowOnlyInnerProperties))
 	FDMXEntityFixturePatchRef FixturePatchRef;
