DMX: Improve positioning of cells when adding selected patches to pixel mapping, fix cells stop resizing by mouse when either side is 1px

diff --git a/Engine/Plugins/VirtualProduction/DMX/DMXPixelMapping/Source/DMXPixelMappingEditor/Private/ViewModels/DMXPixelMappingDMXLibraryViewModel.cpp b/Engine/Plugins/VirtualProduction/DMX/DMXPixelMapping/Source/DMXPixelMappingEditor/Private/ViewModels/DMXPixelMappingDMXLibraryViewModel.cpp
--- a/Engine/Plugins/VirtualProduction/DMX/DMXPixelMapping/Source/DMXPixelMappingEditor/Private/ViewModels/DMXPixelMappingDMXLibraryViewModel.cpp
+++ b/Engine/Plugins/VirtualProduction/DMX/DMXPixelMapping/Source/DMXPixelMappingEditor/Private/ViewModels/DMXPixelMappingDMXLibraryViewModel.cpp
@@ -395,6 +395,7 @@ void UDMXPixelMappingDMXLibraryViewModel::LayoutAfterLastPatch(const TArray<UDMX
 	{
 		if (UDMXPixelMappingOutputComponent* OutputComponent = Cast<UDMXPixelMappingOutputComponent>(Component))
 		{
+			NextPosition = NextPosition - FLT_EPSILON * 2.f;
 			OutputComponent->SetPosition(NextPosition);
 
 			if (OutputComponent->IsOverParent())
diff --git a/Engine/Plugins/VirtualProduction/DMX/DMXPixelMapping/Source/DMXPixelMappingEditor/Private/Widgets/SDMXPixelMappingTransformHandle.cpp b/Engine/Plugins/VirtualProduction/DMX/DMXPixelMapping/Source/DMXPixelMappingEditor/Private/Widgets/SDMXPixelMappingTransformHandle.cpp
--- a/Engine/Plugins/VirtualProduction/DMX/DMXPixelMapping/Source/DMXPixelMappingEditor/Private/Widgets/SDMXPixelMappingTransformHandle.cpp
+++ b/Engine/Plugins/VirtualProduction/DMX/DMXPixelMapping/Source/DMXPixelMappingEditor/Private/Widgets/SDMXPixelMappingTransformHandle.cpp
@@ -257,14 +257,7 @@ void SDMXPixelMappingTransformHandle::Resize(UDMXPixelMappingOutputComponent* Ou
 
 	const FVector2D NewSize = GetSnapSize(OutputComponent, RequestedSize);
 	const FVector2D OldSize = OutputComponent->GetSize();
-	if ((FMath::IsNearlyEqual(OldSize.X, NewSize.X) && FMath::IsNearlyEqual(OldSize.Y, NewSize.Y)) ||
-		NewSize.X < 1.f ||
-		NewSize.Y < 1.f)
-	{
-		// No (nearly) unchanged values, no size < 1 on either axis, no negative values.
-		return;
-	}
-	
+
 	OutputComponent->Modify();
 	OutputComponent->SetSize(NewSize);
 }
@@ -274,14 +267,14 @@ FVector2D SDMXPixelMappingTransformHandle::GetSnapSize(UDMXPixelMappingOutputCom
 	const TSharedPtr<FDMXPixelMappingToolkit> Toolkit = DesignerViewWeakPtr.Pin()->GetToolkit();
 	if (!Toolkit.IsValid() || !OutputComponent || !OutputComponent->GetRendererComponent())
 	{
-		return RequestedSize;
+		return FVector2D::Clamp(RequestedSize, FVector2D(1.f, 1.f), RequestedSize);
 	}
 	UDMXPixelMappingRendererComponent* RendererComponent = OutputComponent->GetRendererComponent();
 
 	const UDMXPixelMapping* PixelMapping = Toolkit->GetDMXPixelMapping();
 	if (!PixelMapping || !PixelMapping->bGridSnappingEnabled)
 	{
-		return RequestedSize;
+		return FVector2D::Clamp(RequestedSize, FVector2D(1.f, 1.f), RequestedSize);
 	}
 
 	const FVector2D CellSize = [PixelMapping, RendererComponent]()
