DMX: Fix Compile UnrealGame - Add missing with editor directive

diff --git a/Engine/Plugins/VirtualProduction/DMX/DMXEngine/Source/DMXRuntime/Private/Library/DMXLibrary.cpp b/Engine/Plugins/VirtualProduction/DMX/DMXEngine/Source/DMXRuntime/Private/Library/DMXLibrary.cpp
--- a/Engine/Plugins/VirtualProduction/DMX/DMXEngine/Source/DMXRuntime/Private/Library/DMXLibrary.cpp
+++ b/Engine/Plugins/VirtualProduction/DMX/DMXEngine/Source/DMXRuntime/Private/Library/DMXLibrary.cpp
@@ -584,7 +584,7 @@ UDMXMVRGeneralSceneDescription* UDMXLibrary::UpdateGeneralSceneDescription()
 
 	return nullptr;
 }
-#endif // WITH_EDTIOR
+#endif // WITH_EDITOR
 
 #if WITH_EDITOR
 void UDMXLibrary::UpgradeFromControllersToPorts()
diff --git a/Engine/Plugins/VirtualProduction/DMX/DMXPixelMapping/Source/DMXPixelMappingRuntime/Private/LayoutScripts/DMXPixelMappingLayoutScript_LayoutByMVR.cpp b/Engine/Plugins/VirtualProduction/DMX/DMXPixelMapping/Source/DMXPixelMappingRuntime/Private/LayoutScripts/DMXPixelMappingLayoutScript_LayoutByMVR.cpp
--- a/Engine/Plugins/VirtualProduction/DMX/DMXPixelMapping/Source/DMXPixelMappingRuntime/Private/LayoutScripts/DMXPixelMappingLayoutScript_LayoutByMVR.cpp
+++ b/Engine/Plugins/VirtualProduction/DMX/DMXPixelMapping/Source/DMXPixelMappingRuntime/Private/LayoutScripts/DMXPixelMappingLayoutScript_LayoutByMVR.cpp
@@ -77,8 +77,11 @@ TMap<const UDMXMVRFixtureNode*, FDMXPixelMappingLayoutToken> UDMXPixelMappingLay
 	{
 		return Result;
 	}
-
+	
+#if WITH_EDITOR
 	DMXLibrary->UpdateGeneralSceneDescription();
+#endif
+
 	UDMXMVRGeneralSceneDescription* GeneralSceneDescription = DMXLibrary->GetLazyGeneralSceneDescription();
 	if (!GeneralSceneDescription)
 	{
