DMX: Fix issues with undo/redo when deleting an Input Source in Pixel Mapping

diff --git a/Engine/Plugins/VirtualProduction/DMX/DMXPixelMapping/Source/DMXPixelMappingEditor/Private/DMXPixelMappingComponentReference.h b/Engine/Plugins/VirtualProduction/DMX/DMXPixelMapping/Source/DMXPixelMappingEditor/Private/DMXPixelMappingComponentReference.h
--- a/Engine/Plugins/VirtualProduction/DMX/DMXPixelMapping/Source/DMXPixelMappingEditor/Private/DMXPixelMappingComponentReference.h
+++ b/Engine/Plugins/VirtualProduction/DMX/DMXPixelMapping/Source/DMXPixelMappingEditor/Private/DMXPixelMappingComponentReference.h
@@ -4,18 +4,27 @@
 
 #include "Components/DMXPixelMappingBaseComponent.h"
 
+#include "DMXPixelMappingComponentReference.generated.h"
+
 class FDMXPixelMappingToolkit;
 
 /**
  * The Component reference is a useful way to hold onto the selection in a way that allows for up to date access to the current preview object.
  * This is a safe way to communicate between different parts of the pixel mapping editor
  */
-class FDMXPixelMappingComponentReference
+USTRUCT()
+struct FDMXPixelMappingComponentReference
 {
+	GENERATED_BODY()
+
 	friend FDMXPixelMappingToolkit;
 
 public:
-	FDMXPixelMappingComponentReference()
+	FDMXPixelMappingComponentReference() = default;
+
+	FDMXPixelMappingComponentReference(TSharedPtr<FDMXPixelMappingToolkit> InToolkit, UDMXPixelMappingBaseComponent* InComponent)
+		: ToolkitWeakPtr(InToolkit)
+		, Component(InComponent)
 	{}
 
 	/** Checks if widget reference is the same as another component reference, based on the template pointers. */
@@ -40,12 +49,7 @@ public:
 	/** @returns stored component pointer */
 	UDMXPixelMappingBaseComponent* GetComponent() const { return Component.Get(); }
 
-	FDMXPixelMappingComponentReference(TSharedPtr<FDMXPixelMappingToolkit> InToolkit, UDMXPixelMappingBaseComponent* InComponent)
-		: ToolkitWeakPtr(InToolkit)
-		, Component(InComponent)
-	{}
-
-	friend FORCEINLINE uint32 GetTypeHash(const class FDMXPixelMappingComponentReference& ComponentRef)
+	friend FORCEINLINE uint32 GetTypeHash(const FDMXPixelMappingComponentReference& ComponentRef)
 	{
 		return GetTypeHash(ComponentRef.GetComponent());
 	}
@@ -53,6 +57,6 @@ public:
 private:
 	TWeakPtr<FDMXPixelMappingToolkit> ToolkitWeakPtr;
 
+	UPROPERTY(Transient)
 	TWeakObjectPtr<UDMXPixelMappingBaseComponent> Component;
 };
-
diff --git a/Engine/Plugins/VirtualProduction/DMX/DMXPixelMapping/Source/DMXPixelMappingEditor/Private/DMXPixelMappingEditorUtils.h b/Engine/Plugins/VirtualProduction/DMX/DMXPixelMapping/Source/DMXPixelMappingEditor/Private/DMXPixelMappingEditorUtils.h
--- a/Engine/Plugins/VirtualProduction/DMX/DMXPixelMapping/Source/DMXPixelMappingEditor/Private/DMXPixelMappingEditorUtils.h
+++ b/Engine/Plugins/VirtualProduction/DMX/DMXPixelMapping/Source/DMXPixelMappingEditor/Private/DMXPixelMappingEditorUtils.h
@@ -6,17 +6,12 @@
 #include "Templates/SharedPointer.h"
 
 class FArrangedWidget;
-class UDMXPixelMappingRendererComponent;
-
-class FDMXPixelMappingToolkit;
-class UDMXPixelMappingBaseComponent;
-class UDMXPixelMapping;
-class FDMXPixelMappingToolkit;
-class FDMXPixelMappingComponentReference;
-class FMenuBuilder;
-
 class FDragDropEvent;
+class FDMXPixelMappingToolkit;
 class SWidget;
+class UDMXPixelMapping;
+class UDMXPixelMappingBaseComponent;
+class UDMXPixelMappingRendererComponent;
 
 /**
  * Shared Pixel Mapping editor functions
diff --git a/Engine/Plugins/VirtualProduction/DMX/DMXPixelMapping/Source/DMXPixelMappingEditor/Private/DragDrop/DMXPixelMappingDragDropOp.h b/Engine/Plugins/VirtualProduction/DMX/DMXPixelMapping/Source/DMXPixelMappingEditor/Private/DragDrop/DMXPixelMappingDragDropOp.h
--- a/Engine/Plugins/VirtualProduction/DMX/DMXPixelMapping/Source/DMXPixelMappingEditor/Private/DragDrop/DMXPixelMappingDragDropOp.h
+++ b/Engine/Plugins/VirtualProduction/DMX/DMXPixelMapping/Source/DMXPixelMappingEditor/Private/DragDrop/DMXPixelMappingDragDropOp.h
@@ -5,11 +5,10 @@
 #include "DragAndDrop/DecoratedDragDropOp.h"
 
 class FDMXPixelMappingComponentTemplate;
-class FDMXPixelMappingComponentReference;
 class FDMXPixelMappingGroupChildDragDropHelper;
 class FDMXPixelMappingToolkit;
 class UDMXPixelMappingBaseComponent;
-
+struct FDMXPixelMappingComponentReference;
 
 /**
  * This drag drop operation allows Component templates from the palate to be dragged and dropped into the designer
diff --git a/Engine/Plugins/VirtualProduction/DMX/DMXPixelMapping/Source/DMXPixelMappingEditor/Private/Toolkits/DMXPixelMappingToolkit.cpp b/Engine/Plugins/VirtualProduction/DMX/DMXPixelMapping/Source/DMXPixelMappingEditor/Private/Toolkits/DMXPixelMappingToolkit.cpp
--- a/Engine/Plugins/VirtualProduction/DMX/DMXPixelMapping/Source/DMXPixelMappingEditor/Private/Toolkits/DMXPixelMappingToolkit.cpp
+++ b/Engine/Plugins/VirtualProduction/DMX/DMXPixelMapping/Source/DMXPixelMappingEditor/Private/Toolkits/DMXPixelMappingToolkit.cpp
@@ -50,6 +50,8 @@ FDMXPixelMappingToolkit::FDMXPixelMappingToolkit()
 	: AnalyticsProvider("PixelMappingEditor")
 {
 	EditorSettingsDump = TArray<uint8, TFixedAllocator<sizeof(UDMXPixelMappingEditorSettings)>>(reinterpret_cast<const uint8*>(GetDefault<UDMXPixelMappingEditorSettings>()), (int32)sizeof(UDMXPixelMappingEditorSettings));
+
+	Selection = NewObject<UDMXPixelMappingToolkitSelection>(GetTransientPackage(), NAME_None, RF_Transactional);
 }
 
 FDMXPixelMappingToolkit::~FDMXPixelMappingToolkit()
@@ -293,9 +295,19 @@ FDMXPixelMappingComponentReference FDMXPixelMappingToolkit::GetReferenceFromComp
 	return FDMXPixelMappingComponentReference(SharedThis(this), InComponent);
 }
 
+UDMXPixelMappingRendererComponent* FDMXPixelMappingToolkit::GetActiveRendererComponent() const
+{
+	checkf(Selection, TEXT("Unexpected invalid selection object in pixel mapping toolkit."));
+
+	return Selection->ActiveRendererComponent.Get();
+}
+
 void FDMXPixelMappingToolkit::SetActiveRenderComponent(UDMXPixelMappingRendererComponent* InComponent)
 {
-	ActiveRendererComponent = InComponent;
+	checkf(Selection, TEXT("Unexpected invalid selection object in pixel mapping toolkit."));
+
+	Selection->Modify();
+	Selection->ActiveRendererComponent = InComponent;
 }
 
 template <typename ComponentType>
@@ -313,13 +325,16 @@ TArray<ComponentType> FDMXPixelMappingToolkit::MakeComponentArray(const TSet<FDM
 
 void FDMXPixelMappingToolkit::SelectComponents(const TSet<FDMXPixelMappingComponentReference>& InSelectedComponents)
 {
+	checkf(Selection, TEXT("Unexpected invalid selection object in pixel mapping toolkit."));
+
 	// Update selection
-	SelectedComponents.Empty();
-	ActiveOutputComponents.Empty();
+	Selection->Modify();
+	Selection->Components.Reset();
+	ActiveOutputComponents.Reset();
 
-	SelectedComponents.Append(InSelectedComponents);
+	Selection->Components.Append(InSelectedComponents);
 
-	for (const FDMXPixelMappingComponentReference& ComponentReference : SelectedComponents)
+	for (const FDMXPixelMappingComponentReference& ComponentReference : Selection->Components)
 	{
 		if (UDMXPixelMappingRootComponent* RootComponent = Cast<UDMXPixelMappingRootComponent>(ComponentReference.GetComponent()))
 		{
@@ -345,7 +360,7 @@ void FDMXPixelMappingToolkit::SelectComponents(const TSet<FDMXPixelMappingCompon
 
 	// Always order selected components topmost, but keep their relative z-ordering
 	TArray<UDMXPixelMappingOutputComponent*> SelectedOutputComponents;
-	Algo::TransformIf(SelectedComponents, SelectedOutputComponents,
+	Algo::TransformIf(Selection->Components, SelectedOutputComponents,
 		[](const FDMXPixelMappingComponentReference& ComponentReference)
 		{
 			return ComponentReference.GetComponent() && ComponentReference.GetComponent()->IsA(UDMXPixelMappingOutputComponent::StaticClass());
@@ -365,7 +380,7 @@ void FDMXPixelMappingToolkit::SelectComponents(const TSet<FDMXPixelMappingCompon
 
 bool FDMXPixelMappingToolkit::IsComponentSelected(UDMXPixelMappingBaseComponent* Component) const
 {
-	for (const FDMXPixelMappingComponentReference& ComponentReference : SelectedComponents)
+	for (const FDMXPixelMappingComponentReference& ComponentReference : Selection->Components)
 	{
 		if (Component && Component == ComponentReference.GetComponent())
 		{
@@ -472,13 +487,12 @@ void FDMXPixelMappingToolkit::UpdateBlueprintNodes() const
 void FDMXPixelMappingToolkit::SaveThumbnailImage()
 {
 	UDMXPixelMapping* PixelMapping = GetDMXPixelMapping();
-	if (PixelMapping)
+	UDMXPixelMappingRendererComponent* ActiveRendererComponent = GetActiveRendererComponent();
+
+	UTexture* Texture = ActiveRendererComponent ? ActiveRendererComponent->GetRenderedInputTexture() : nullptr;
+	if (IsValid(Texture) && Texture->IsFullyStreamedIn())
 	{
-		UTexture* Texture = ActiveRendererComponent.IsValid() ? ActiveRendererComponent->GetRenderedInputTexture() : nullptr;
-		if (IsValid(Texture) && Texture->IsFullyStreamedIn())
-		{
-			PixelMapping->ThumbnailImage = Texture;
-		}
+		PixelMapping->ThumbnailImage = Texture;
 	}
 }
 
@@ -534,7 +548,9 @@ TArray<UDMXPixelMappingBaseComponent*> FDMXPixelMappingToolkit::CreateComponents
 
 void FDMXPixelMappingToolkit::DeleteSelectedComponents()
 {
-	if (SelectedComponents.IsEmpty())
+	checkf(Selection, TEXT("Unexpected invalid selection object in pixel mapping toolkit."));
+
+	if (Selection->Components.IsEmpty())
 	{
 		return;
 	}
@@ -542,7 +558,7 @@ void FDMXPixelMappingToolkit::DeleteSelectedComponents()
 	TGuardValue Guard(bRemovingComponents, true);
 
 	TSet<FDMXPixelMappingComponentReference> ParentComponentReferences;
-	for (const FDMXPixelMappingComponentReference& SelectedComponentReference : SelectedComponents)
+	for (const FDMXPixelMappingComponentReference& SelectedComponentReference : Selection->Components)
 	{
 		UDMXPixelMappingBaseComponent* SelectedComponent = SelectedComponentReference.GetComponent();
 		if (SelectedComponent)
@@ -561,11 +577,7 @@ void FDMXPixelMappingToolkit::DeleteSelectedComponents()
 
 				ParentComponent->RemoveChild(SelectedComponent);
 
-				const bool bParentComponentIsBeingRemoved = Algo::FindByPredicate(SelectedComponents, [ParentComponent](const FDMXPixelMappingComponentReference& Reference)
-					{
-						return Reference.GetComponent() == ParentComponent;
-					}) != nullptr;
-
+				const bool bParentComponentIsBeingRemoved = Algo::FindBy(Selection->Components, ParentComponent, &FDMXPixelMappingComponentReference::GetComponent) != nullptr;
 				if (!bParentComponentIsBeingRemoved)
 				{
 					ParentComponentReferences.Add(FDMXPixelMappingComponentReference(StaticCastSharedRef<FDMXPixelMappingToolkit>(AsShared()), ParentComponent));
@@ -687,6 +699,16 @@ void FDMXPixelMappingToolkit::ToggleGridSnapping()
 	}
 }
 
+void FDMXPixelMappingToolkit::AddReferencedObjects(FReferenceCollector& Collector)
+{
+	Collector.AddReferencedObject(Selection);
+}
+
+FString FDMXPixelMappingToolkit::GetReferencerName() const
+{
+	return TEXT("FDMXPixelMappingToolkit");
+}
+
 void FDMXPixelMappingToolkit::PostUndo(bool bSuccess)
 {
 	UDMXPixelMapping* PixelMapping = GetDMXPixelMapping();
@@ -839,6 +861,13 @@ void FDMXPixelMappingToolkit::RenameComponent(const FName& CurrentObjectName, co
 	UpdateBlueprintNodes();
 }
 
+const TSet<FDMXPixelMappingComponentReference>& FDMXPixelMappingToolkit::GetSelectedComponents() const
+{
+	checkf(Selection, TEXT("Unexpected invalid selection object in pixel mapping toolkit."));
+
+	return Selection->Components;
+}
+
 TSharedRef<SDMXPixelMappingDMXLibraryView> FDMXPixelMappingToolkit::GetOrCreateDMXLibraryView()
 {
 	if (!DMXLibraryView.IsValid())
@@ -1093,8 +1122,10 @@ ECheckBoxState FDMXPixelMappingToolkit::GetTransformHandleModeCheckboxState(EDMX
 
 UDMXPixelMappingFixtureGroupComponent* FDMXPixelMappingToolkit::GetFixtureGroupFromSelection() const
 {
+	checkf(Selection, TEXT("Unexpected invalid selection object in pixel mapping toolkit."));
+
 	TArray<UDMXPixelMappingFixtureGroupComponent*> FixtureGroupComponents;
-	for (const FDMXPixelMappingComponentReference& ComponentReference : SelectedComponents)
+	for (const FDMXPixelMappingComponentReference& ComponentReference : Selection->Components)
 	{
 		UDMXPixelMappingBaseComponent* Component = ComponentReference.GetComponent();
 		if (!Component)
diff --git a/Engine/Plugins/VirtualProduction/DMX/DMXPixelMapping/Source/DMXPixelMappingEditor/Private/Toolkits/DMXPixelMappingToolkit.h b/Engine/Plugins/VirtualProduction/DMX/DMXPixelMapping/Source/DMXPixelMappingEditor/Private/Toolkits/DMXPixelMappingToolkit.h
--- a/Engine/Plugins/VirtualProduction/DMX/DMXPixelMapping/Source/DMXPixelMappingEditor/Private/Toolkits/DMXPixelMappingToolkit.h
+++ b/Engine/Plugins/VirtualProduction/DMX/DMXPixelMapping/Source/DMXPixelMappingEditor/Private/Toolkits/DMXPixelMappingToolkit.h
@@ -11,6 +11,8 @@
 #include "UObject/GCObject.h"
 #include "Widgets/Views/SHeaderRow.h"
 
+#include "DMXPixelMappingToolkit.generated.h"
+
 enum class ECheckBoxState : uint8;
 enum class EDMXPixelMappingResetDMXMode : uint8;
 class FDMXPixelMappingComponentTemplate;
@@ -40,6 +42,20 @@ namespace UE::DMX
 	};
 }
 
+/** Selection in the pixel mapping toollkit. Using the UObject system for undo/redo. */
+UCLASS(Transient)
+class UDMXPixelMappingToolkitSelection
+	: public UObject
+{
+	GENERATED_BODY()
+
+public:
+	UPROPERTY()
+	TSet<FDMXPixelMappingComponentReference> Components;
+
+	UPROPERTY()
+	TWeakObjectPtr<UDMXPixelMappingRendererComponent> ActiveRendererComponent;
+};
 
 /**
  * Implements an Editor toolkit for Pixel Mapping.
@@ -48,6 +64,7 @@ class FDMXPixelMappingToolkit
 	: public FAssetEditorToolkit
 	, public FTickableEditorObject
 	, public FSelfRegisteringEditorUndoClient
+	, public FGCObject
 {
 	using EDMXPixelMappingTransformHandleMode = UE::DMX::EDMXPixelMappingTransformHandleMode;
 
@@ -91,16 +108,16 @@ public:
 	virtual bool IsTickable() const override { return true; }
 	virtual TStatId GetStatId() const override;
 	//~ End FTickableEditorObject Interface
-		
+
 	UDMXPixelMapping* GetDMXPixelMapping() const;
 
 	FDMXPixelMappingComponentReference GetReferenceFromComponent(UDMXPixelMappingBaseComponent* InComponent);
 
-	UDMXPixelMappingRendererComponent* GetActiveRendererComponent() const { return ActiveRendererComponent.Get(); }
+	UDMXPixelMappingRendererComponent* GetActiveRendererComponent() const;
 
 	const TSharedPtr<FUICommandList>& GetDesignerCommandList() const { return DesignerCommandList; }
 
-	const TSet<FDMXPixelMappingComponentReference>& GetSelectedComponents() const { return SelectedComponents; }
+	const TSet<FDMXPixelMappingComponentReference>& GetSelectedComponents() const;
 
 	/** Gets or creates the DMX Library View for this Pixel Mapping instance */
 	TSharedRef<SDMXPixelMappingDMXLibraryView> GetOrCreateDMXLibraryView();
@@ -166,6 +183,12 @@ public:
 	/** Returns the current transform handle mode */
 	EDMXPixelMappingTransformHandleMode GetTransformHandleMode() const { return TransformHandleMode; }
 
+protected:
+	//~ Begin FGCObject interface
+	virtual void AddReferencedObjects(FReferenceCollector& Collector) override;
+	virtual FString GetReferencerName() const override;
+	//~ End FGCObject interface
+
 private:
 	//~ Begin FSelfRegisteringEditorUndoClient interface
 	virtual void PostUndo(bool bSuccess) override;
@@ -261,12 +284,11 @@ private:
 
 	FOnSelectedComponentsChangedDelegate OnSelectedComponentsChangedDelegate;
 
-	TSet<FDMXPixelMappingComponentReference> SelectedComponents;
+	/** The selection object */
+	TObjectPtr<UDMXPixelMappingToolkitSelection> Selection;
 
 	TSharedPtr<FDMXPixelMappingToolbar> Toolbar;
 
-	TWeakObjectPtr<UDMXPixelMappingRendererComponent> ActiveRendererComponent;
-
 	TArray<TWeakObjectPtr<UDMXPixelMappingOutputComponent>> ActiveOutputComponents;
 
 	/** Command list for handling widget actions in the PixelMapping Toolkit */
diff --git a/Engine/Plugins/VirtualProduction/DMX/DMXPixelMapping/Source/DMXPixelMappingEditor/Private/Views/SDMXPixelMappingDesignerView.cpp b/Engine/Plugins/VirtualProduction/DMX/DMXPixelMapping/Source/DMXPixelMappingEditor/Private/Views/SDMXPixelMappingDesignerView.cpp
--- a/Engine/Plugins/VirtualProduction/DMX/DMXPixelMapping/Source/DMXPixelMappingEditor/Private/Views/SDMXPixelMappingDesignerView.cpp
+++ b/Engine/Plugins/VirtualProduction/DMX/DMXPixelMapping/Source/DMXPixelMappingEditor/Private/Views/SDMXPixelMappingDesignerView.cpp
@@ -650,12 +650,21 @@ int32 SDMXPixelMappingDesignerView::GetGridSize() const
 
 void SDMXPixelMappingDesignerView::PostUndo(bool bSuccess)
 {
-	CachedRendererComponent = nullptr;
+	const TSharedPtr<FDMXPixelMappingToolkit> Toolkit = WeakToolkit.Pin();
+	if (!Toolkit.IsValid())
+	{
+		return;
+	}
+
+	UDMXPixelMappingRendererComponent* RendererComponent = Toolkit->GetActiveRendererComponent();
+	CachedRendererComponent = RendererComponent;
+	RebuildDesigner();
 }
 
 void SDMXPixelMappingDesignerView::PostRedo(bool bSuccess)
 {
-	CachedRendererComponent = nullptr;
+	// Same as post undo
+	PostUndo(bSuccess);
 }
 
 void SDMXPixelMappingDesignerView::RebuildDesigner()
diff --git a/Engine/Plugins/VirtualProduction/DMX/DMXPixelMapping/Source/DMXPixelMappingEditor/Private/Views/SDMXPixelMappingDesignerView.h b/Engine/Plugins/VirtualProduction/DMX/DMXPixelMapping/Source/DMXPixelMappingEditor/Private/Views/SDMXPixelMappingDesignerView.h
--- a/Engine/Plugins/VirtualProduction/DMX/DMXPixelMapping/Source/DMXPixelMappingEditor/Private/Views/SDMXPixelMappingDesignerView.h
+++ b/Engine/Plugins/VirtualProduction/DMX/DMXPixelMapping/Source/DMXPixelMappingEditor/Private/Views/SDMXPixelMappingDesignerView.h
@@ -5,7 +5,6 @@
 #include "EditorUndoClient.h"
 #include "Widgets/SDMXPixelMappingSurface.h"
 
-class FDMXPixelMappingComponentReference;
 class FDMXPixelMappingDragDropOp;
 class FDMXPixelMappingToolkit;
 class SBorder;
@@ -21,11 +20,8 @@ class UDMXPixelMapping;
 class UDMXPixelMappingBaseComponent;
 class UDMXPixelMappingOutputComponent;
 class UDMXPixelMappingRendererComponent;
-namespace UE::DMX
-{
-	class IDMXPixelMappingOutputComponentWidgetInterface;
-}
-
+namespace UE::DMX { class IDMXPixelMappingOutputComponentWidgetInterface; }
+struct FDMXPixelMappingComponentReference;
 
 class SDMXPixelMappingDesignerView
 	: public SDMXPixelMappingSurface
@@ -44,6 +40,7 @@ private:
 			: WidgetArranged(SNullWidget::NullWidget, FGeometry())
 		{}
 	};
+
 public:
 	SLATE_BEGIN_ARGS(SDMXPixelMappingDesignerView) 
 	{}
diff --git a/Engine/Plugins/VirtualProduction/DMX/DMXPixelMapping/Source/DMXPixelMappingEditor/Private/Views/SDMXPixelMappingPreviewView.h b/Engine/Plugins/VirtualProduction/DMX/DMXPixelMapping/Source/DMXPixelMappingEditor/Private/Views/SDMXPixelMappingPreviewView.h
--- a/Engine/Plugins/VirtualProduction/DMX/DMXPixelMapping/Source/DMXPixelMappingEditor/Private/Views/SDMXPixelMappingPreviewView.h
+++ b/Engine/Plugins/VirtualProduction/DMX/DMXPixelMapping/Source/DMXPixelMappingEditor/Private/Views/SDMXPixelMappingPreviewView.h
@@ -2,16 +2,16 @@
 
 #pragma once
 
+#include "EditorUndoClient.h"
 #include "Widgets/SDMXPixelMappingSurface.h"
 
-class FDMXPixelMappingToolkit;
+struct FDMXPixelMappingComponentReference;
 class SDMXPixelMappingPreviewViewport;
+class FDMXPixelMappingToolkit;
 struct FOptionalSize;
 class SBox;
 class SDMXPixelMappingZoomPan;
 class SOverlay;
-class FDMXPixelMappingComponentReference;
-
 
 class SDMXPixelMappingPreviewView
 	: public SDMXPixelMappingSurface
