DMX: Fully deprecate the Pixel Mapping Screen Component

diff --git a/Engine/Plugins/VirtualProduction/DMX/DMXPixelMapping/Source/DMXPixelMappingEditor/Private/Customizations/DMXPixelMappingDetailCustomization_Screen.cpp b/Engine/Plugins/VirtualProduction/DMX/DMXPixelMapping/Source/DMXPixelMappingEditor/Private/Customizations/DMXPixelMappingDetailCustomization_Screen.cpp
deleted file mode 100644
diff --git a/Engine/Plugins/VirtualProduction/DMX/DMXPixelMapping/Source/DMXPixelMappingEditor/Private/Customizations/DMXPixelMappingDetailCustomization_Screen.h b/Engine/Plugins/VirtualProduction/DMX/DMXPixelMapping/Source/DMXPixelMappingEditor/Private/Customizations/DMXPixelMappingDetailCustomization_Screen.h
deleted file mode 100644
diff --git a/Engine/Plugins/VirtualProduction/DMX/DMXPixelMapping/Source/DMXPixelMappingEditor/Private/Toolkits/DMXPixelMappingToolkit.cpp b/Engine/Plugins/VirtualProduction/DMX/DMXPixelMapping/Source/DMXPixelMappingEditor/Private/Toolkits/DMXPixelMappingToolkit.cpp
--- a/Engine/Plugins/VirtualProduction/DMX/DMXPixelMapping/Source/DMXPixelMappingEditor/Private/Toolkits/DMXPixelMappingToolkit.cpp
+++ b/Engine/Plugins/VirtualProduction/DMX/DMXPixelMapping/Source/DMXPixelMappingEditor/Private/Toolkits/DMXPixelMappingToolkit.cpp
@@ -10,7 +10,6 @@
 #include "Components/DMXPixelMappingRootComponent.h"
 #include "Components/DMXPixelMappingMatrixCellComponent.h"
 #include "Components/DMXPixelMappingMatrixComponent.h"
-#include "Components/DMXPixelMappingScreenComponent.h"
 #include "DMXPixelMapping.h"
 #include "DMXPixelMappingEditorCommands.h"
 #include "DMXPixelMappingEditorModule.h"
diff --git a/Engine/Plugins/VirtualProduction/DMX/DMXPixelMapping/Source/DMXPixelMappingEditor/Private/ViewModels/DMXPixelMappingLayoutViewModel.cpp b/Engine/Plugins/VirtualProduction/DMX/DMXPixelMapping/Source/DMXPixelMappingEditor/Private/ViewModels/DMXPixelMappingLayoutViewModel.cpp
--- a/Engine/Plugins/VirtualProduction/DMX/DMXPixelMapping/Source/DMXPixelMappingEditor/Private/ViewModels/DMXPixelMappingLayoutViewModel.cpp
+++ b/Engine/Plugins/VirtualProduction/DMX/DMXPixelMapping/Source/DMXPixelMappingEditor/Private/ViewModels/DMXPixelMappingLayoutViewModel.cpp
@@ -2,19 +2,17 @@
 
 #include "ViewModels/DMXPixelMappingLayoutViewModel.h"
 
-#include "Components/DMXPixelMappingMatrixCellComponent.h"
 #include "Components/DMXPixelMappingFixtureGroupComponent.h"
 #include "Components/DMXPixelMappingFixtureGroupItemComponent.h"
+#include "Components/DMXPixelMappingMatrixCellComponent.h"
 #include "Components/DMXPixelMappingMatrixComponent.h"
 #include "Components/DMXPixelMappingRendererComponent.h"
-#include "Components/DMXPixelMappingScreenComponent.h"
 #include "DMXPixelMappingEditorLog.h"
 #include "LayoutScripts/DMXPixelMappingLayoutScript.h"
 #include "ScopedTransaction.h"
 #include "Settings/DMXPixelMappingEditorSettings.h"
 #include "Toolkits/DMXPixelMappingToolkit.h"
 
-
 #define LOCTEXT_NAMESPACE "DMXPixelMappingLayoutViewModel"
 
 void UDMXPixelMappingLayoutViewModel::PreEditChange(FProperty* PropertyAboutToChange)
@@ -59,16 +57,11 @@ EDMXPixelMappingLayoutViewModelMode UDMXPixelMappingLayoutViewModel::GetMode() c
 {
 	// If a Fixture Group is selected, layout its children.
 	// If a Matrix is selected, layout its children.
-	// If more than one Fixture Group or a Screen Component is selected, layout those.
-	if (FixtureGroupComponents.Num() == 1 && ScreenComponents.IsEmpty())
-	{
-		return EDMXPixelMappingLayoutViewModelMode::LayoutFixtureGroupComponentChildren;
-	}
-	else if (MatrixComponents.Num() == 1 && FixtureGroupComponents.IsEmpty())
+	if (MatrixComponents.Num() == 1 && FixtureGroupComponents.IsEmpty())
 	{
 		return EDMXPixelMappingLayoutViewModelMode::LayoutMatrixComponentChildren;
 	}
-	else if (RendererComponent.IsValid() && ScreenComponents.IsEmpty() && FixtureGroupComponents.IsEmpty() && MatrixComponents.IsEmpty())
+	else if (RendererComponent.IsValid() && FixtureGroupComponents.IsEmpty() && MatrixComponents.IsEmpty())
 	{
 		return EDMXPixelMappingLayoutViewModelMode::LayoutRendererComponentChildren;
 	}
@@ -166,7 +159,7 @@ UDMXPixelMappingOutputComponent* UDMXPixelMappingLayoutViewModel::GetParentCompo
 	}
 	else if (LayoutMode == EDMXPixelMappingLayoutViewModelMode::LayoutFixtureGroupComponentChildren)
 	{
-		if (!ensureMsgf(FixtureGroupComponents.Num() == 1 && ScreenComponents.IsEmpty(), TEXT("GetMode no longer matches assumed conditions.")))
+		if (!ensureMsgf(FixtureGroupComponents.Num() == 1, TEXT("GetMode no longer matches assumed conditions.")))
 		{
 			return nullptr;
 		}
@@ -463,7 +456,6 @@ void UDMXPixelMappingLayoutViewModel::RefreshComponents()
 {
 	// Reset
 	RendererComponent = nullptr;
-	ScreenComponents.Reset();
 	FixtureGroupComponents.Reset();
 	MatrixComponents.Reset();
 
@@ -486,11 +478,7 @@ void UDMXPixelMappingLayoutViewModel::RefreshComponents()
 	{
 		if (UDMXPixelMappingBaseComponent* Component = ComponentReference.GetComponent())
 		{
-			if (UDMXPixelMappingScreenComponent* ScreenComponent = Cast<UDMXPixelMappingScreenComponent>(Component))
-			{
-				ScreenComponents.Add(ScreenComponent);
-			}
-			else if (UDMXPixelMappingFixtureGroupComponent* FixtureGroupComponent = Cast<UDMXPixelMappingFixtureGroupComponent>(Component))
+			if (UDMXPixelMappingFixtureGroupComponent* FixtureGroupComponent = Cast<UDMXPixelMappingFixtureGroupComponent>(Component))
 			{
 				FixtureGroupComponents.Add(FixtureGroupComponent);
 			}
diff --git a/Engine/Plugins/VirtualProduction/DMX/DMXPixelMapping/Source/DMXPixelMappingEditor/Private/ViewModels/DMXPixelMappingLayoutViewModel.h b/Engine/Plugins/VirtualProduction/DMX/DMXPixelMapping/Source/DMXPixelMappingEditor/Private/ViewModels/DMXPixelMappingLayoutViewModel.h
--- a/Engine/Plugins/VirtualProduction/DMX/DMXPixelMapping/Source/DMXPixelMappingEditor/Private/ViewModels/DMXPixelMappingLayoutViewModel.h
+++ b/Engine/Plugins/VirtualProduction/DMX/DMXPixelMapping/Source/DMXPixelMappingEditor/Private/ViewModels/DMXPixelMappingLayoutViewModel.h
@@ -3,8 +3,8 @@
 #pragma once
 
 #include "EditorUndoClient.h"
-
 #include "Engine/TimerHandle.h"
+
 #include "DMXPixelMappingLayoutViewModel.generated.h"
 
 struct FDMXPixelMappingLayoutToken;
@@ -15,8 +15,6 @@ class UDMXPixelMappingLayoutScript;
 class UDMXPixelMappingMatrixComponent;
 class UDMXPixelMappingOutputComponent;
 class UDMXPixelMappingRendererComponent;
-class UDMXPixelMappingScreenComponent;
-
 
 enum class EDMXPixelMappingLayoutViewModelMode
 {
@@ -26,7 +24,6 @@ enum class EDMXPixelMappingLayoutViewModelMode
 	LayoutNone
 };
 
-
 /** Model for the Layout View */
 UCLASS(Transient)
 class UDMXPixelMappingLayoutViewModel
@@ -117,10 +114,6 @@ private:
 	UPROPERTY()
 	TWeakObjectPtr<UDMXPixelMappingRendererComponent> RendererComponent;
 
-	/** The Screen Components currently in use */
-	UPROPERTY()
-	TArray<TWeakObjectPtr<UDMXPixelMappingScreenComponent>> ScreenComponents;
-
 	/** The Fixture Group Components currently in use */
 	UPROPERTY()
 	TArray<TWeakObjectPtr<UDMXPixelMappingFixtureGroupComponent>> FixtureGroupComponents;
diff --git a/Engine/Plugins/VirtualProduction/DMX/DMXPixelMapping/Source/DMXPixelMappingEditor/Private/ViewModels/DMXPixelMappingOutputComponentModel.cpp b/Engine/Plugins/VirtualProduction/DMX/DMXPixelMapping/Source/DMXPixelMappingEditor/Private/ViewModels/DMXPixelMappingOutputComponentModel.cpp
--- a/Engine/Plugins/VirtualProduction/DMX/DMXPixelMapping/Source/DMXPixelMappingEditor/Private/ViewModels/DMXPixelMappingOutputComponentModel.cpp
+++ b/Engine/Plugins/VirtualProduction/DMX/DMXPixelMapping/Source/DMXPixelMappingEditor/Private/ViewModels/DMXPixelMappingOutputComponentModel.cpp
@@ -3,12 +3,11 @@
 #include "ViewModels/DMXPixelMappingOutputComponentModel.h"
 
 #include "Components/DMXPixelMappingFixtureGroupComponent.h"
-#include "DMXPixelMappingTypes.h"
 #include "Components/DMXPixelMappingFixtureGroupItemComponent.h"
 #include "Components/DMXPixelMappingMatrixCellComponent.h"
 #include "Components/DMXPixelMappingMatrixComponent.h"
 #include "Components/DMXPixelMappingOutputComponent.h"
-#include "Components/DMXPixelMappingScreenComponent.h"
+#include "DMXPixelMappingTypes.h"
 #include "Library/DMXEntityFixturePatch.h"
 #include "Library/DMXLibrary.h"
 #include "MVR/DMXMVRGeneralSceneDescription.h"
@@ -113,9 +112,7 @@ namespace UE::DMX
 	bool FDMXPixelMappingOutputComponentModel::ShouldDrawNameAbove() const
 	{
 		return
-			OutputComponent &&
-			(OutputComponent->GetClass() == UDMXPixelMappingFixtureGroupComponent::StaticClass() ||
-			OutputComponent->GetClass() == UDMXPixelMappingScreenComponent::StaticClass());
+			OutputComponent && OutputComponent->GetClass() == UDMXPixelMappingFixtureGroupComponent::StaticClass();
 	}
 
 	bool FDMXPixelMappingOutputComponentModel::ShouldDrawCellID() const
@@ -270,158 +267,6 @@ namespace UE::DMX
 
 		WeakFixtureNode = GeneralSceneDescription->FindFixtureNode(MVRUUID);
 	}
-
-	FDMXPixelMappingScreenComponentModel::FDMXPixelMappingScreenComponentModel(const TSharedRef<FDMXPixelMappingToolkit>& InToolkit, TWeakObjectPtr<UDMXPixelMappingScreenComponent> InScreenComponent)
-		: WeakScreenComponent(InScreenComponent)
-		, WeakToolkit(InToolkit)
-	{
-		UDMXPixelMappingScreenComponent* ScreenComponent = WeakScreenComponent.Get();
-		if (!ScreenComponent)
-		{
-			return;
-		}
-
-		bSelected = InToolkit->GetSelectedComponents().Contains(FDMXPixelMappingComponentReference(InToolkit, ScreenComponent));
-
-		// Handle selection changes
-		InToolkit->GetOnSelectedComponentsChangedDelegate().AddRaw(this, &FDMXPixelMappingScreenComponentModel::OnSelectedComponentsChanged);
-	}
-
-	FDMXPixelMappingScreenComponentModel::~FDMXPixelMappingScreenComponentModel()
-	{
-		if (TSharedPtr<FDMXPixelMappingToolkit> Toolkit = WeakToolkit.Pin())
-		{
-			Toolkit->GetOnSelectedComponentsChangedDelegate().RemoveAll(this);
-		}
-	}
-
-	FVector2D FDMXPixelMappingScreenComponentModel::GetPosition() const
-	{
-		if (UDMXPixelMappingScreenComponent* ScreenComponent = WeakScreenComponent.Get())
-		{
-			return ScreenComponent->GetPosition();
-		}
-		return FVector2D::ZeroVector;
-	}
-
-	FVector2D FDMXPixelMappingScreenComponentModel::GetSize() const
-	{
-		if (UDMXPixelMappingScreenComponent* ScreenComponent = WeakScreenComponent.Get())
-		{
-			return ScreenComponent->GetSize();
-		}
-		return FVector2D::ZeroVector;
-	}
-
-	int32 FDMXPixelMappingScreenComponentModel::GetNumColumns() const
-	{
-		if (UDMXPixelMappingScreenComponent* ScreenComponent = WeakScreenComponent.Get())
-		{
-			return ScreenComponent->NumXCells;
-		}
-		return 1;
-	}
-
-	int32 FDMXPixelMappingScreenComponentModel::GetNumRows() const
-	{
-		if (UDMXPixelMappingScreenComponent* ScreenComponent = WeakScreenComponent.Get())
-		{
-			return ScreenComponent->NumYCells;
-		}
-		return 1;
-	}
-
-	FLinearColor FDMXPixelMappingScreenComponentModel::GetColor() const
-	{
-		if (UDMXPixelMappingScreenComponent* ScreenComponent = WeakScreenComponent.Get())
-		{
-			if (bSelected)
-			{
-				return ScreenComponent->GetEditorColor();
-			}
-			else
-			{
-				return ScreenComponent->GetEditorColor().CopyWithNewOpacity(0.6f);
-			}
-		}
-
-		return FLinearColor::Red;
-	}
-
-	bool FDMXPixelMappingScreenComponentModel::Equals(UDMXPixelMappingBaseComponent* Other) const
-	{
-		return WeakScreenComponent.Get() == Other;
-	}
-
-	void FDMXPixelMappingScreenComponentModel::OnSelectedComponentsChanged()
-	{
-		TSharedPtr<FDMXPixelMappingToolkit> Toolkit = WeakToolkit.Pin();
-		if (!Toolkit.IsValid())
-		{
-			return;
-		}
-
-		if (UDMXPixelMappingScreenComponent* ScreenComponent = WeakScreenComponent.Get())
-		{
-			bSelected = Toolkit->GetSelectedComponents().Contains(FDMXPixelMappingComponentReference(Toolkit, ScreenComponent));
-		}
-	}
-
-	EDMXPixelMappingDistribution FDMXPixelMappingScreenComponentModel::GetDistribution() const
-	{
-		if (UDMXPixelMappingScreenComponent* ScreenComponent = WeakScreenComponent.Get())
-		{
-			return ScreenComponent->Distribution;
-		}
-		return EDMXPixelMappingDistribution::BottomLeftToRight;
-	}
-
-	EDMXCellFormat FDMXPixelMappingScreenComponentModel::GetCellFormat() const
-	{
-		if (UDMXPixelMappingScreenComponent* ScreenComponent = WeakScreenComponent.Get())
-		{
-			return ScreenComponent->PixelFormat;
-		}
-		return EDMXCellFormat::PF_RGB;
-	}
-
-	bool FDMXPixelMappingScreenComponentModel::ComponentWantsToShowUniverse() const
-	{
-		if (UDMXPixelMappingScreenComponent* ScreenComponent = WeakScreenComponent.Get())
-		{
-			return ScreenComponent->bShowUniverse;
-		}
-		return false;
-	}
-
-	int32 FDMXPixelMappingScreenComponentModel::GetUniverse() const
-	{
-		if (UDMXPixelMappingScreenComponent* ScreenComponent = WeakScreenComponent.Get())
-		{
-			return ScreenComponent->LocalUniverse;
-		}
-
-		return 0;
-	}
-
-	bool FDMXPixelMappingScreenComponentModel::ComponentWantsToShowChannel() const
-	{
-		if (UDMXPixelMappingScreenComponent* ScreenComponent = WeakScreenComponent.Get())
-		{
-			return ScreenComponent->bShowAddresses;
-		}
-		return false;
-	}
-
-	int32 FDMXPixelMappingScreenComponentModel::GetStartingChannel() const
-	{
-		if (UDMXPixelMappingScreenComponent* ScreenComponent = WeakScreenComponent.Get())
-		{
-			return ScreenComponent->StartAddress;
-		}
-
-		return 0;
-	}
 }
 
 #undef LOCTEXT_NAMESPACE
diff --git a/Engine/Plugins/VirtualProduction/DMX/DMXPixelMapping/Source/DMXPixelMappingEditor/Private/ViewModels/DMXPixelMappingOutputComponentModel.h b/Engine/Plugins/VirtualProduction/DMX/DMXPixelMapping/Source/DMXPixelMappingEditor/Private/ViewModels/DMXPixelMappingOutputComponentModel.h
--- a/Engine/Plugins/VirtualProduction/DMX/DMXPixelMapping/Source/DMXPixelMappingEditor/Private/ViewModels/DMXPixelMappingOutputComponentModel.h
+++ b/Engine/Plugins/VirtualProduction/DMX/DMXPixelMapping/Source/DMXPixelMappingEditor/Private/ViewModels/DMXPixelMappingOutputComponentModel.h
@@ -14,15 +14,11 @@ class UDMXEntityFixturePatch;
 class UDMXMVRFixtureNode;
 class UDMXPixelMappingBaseComponent;
 class UDMXPixelMappingOutputComponent;
-class UDMXPixelMappingScreenComponent;
 
 
 namespace UE::DMX
 {
-	/**
-	 * Model for the Pixel Mapping Output Components.
-	 * For Screen Component, see DMXPixelMappingScreenComponentModel.
-	 */
+	/** Model for the Pixel Mapping Output Components. For Screen Component, see DMXPixelMappingScreenComponentModel. */
 	class FDMXPixelMappingOutputComponentModel
 		: public FGCObject
 		, public TSharedFromThis<FDMXPixelMappingOutputComponentModel>
@@ -110,66 +106,4 @@ namespace UE::DMX
 		/** The toolkit from which the Model currently sources */
 		TWeakPtr<FDMXPixelMappingToolkit> WeakToolkit;
 	};
-
-
-	/** Model for Pixel Mapping Screen Components */
-	class FDMXPixelMappingScreenComponentModel
-		: public TSharedFromThis<FDMXPixelMappingScreenComponentModel>
-	{
-	public:
-		/** Constructor */
-		FDMXPixelMappingScreenComponentModel(const TSharedRef<FDMXPixelMappingToolkit>& InToolkit, TWeakObjectPtr<UDMXPixelMappingScreenComponent> InScreenComponent);
-
-		/** Destructor */
-		virtual ~FDMXPixelMappingScreenComponentModel();
-
-		/** Returns the position of the component */
-		FVector2D GetPosition() const;
-
-		/** Returns the size of the component */
-		FVector2D GetSize() const;
-
-		/** Returns the number of Columns of the Screen Component */
-		int32 GetNumColumns() const;
-
-		/** Returns the number of Rows of the Screen Component */
-		int32 GetNumRows() const;
-
-		/** Returns the pixel mapping distribution of the cells */
-		EDMXPixelMappingDistribution GetDistribution() const;
-
-		/** Returns the color format of the cells */
-		EDMXCellFormat GetCellFormat() const;
-
-		/** Returns true if the components wants to show the universe (it has a related property) */
-		bool ComponentWantsToShowUniverse() const;
-
-		/** Returns the universe of the screen component */
-		int32 GetUniverse() const;
-
-		/** Returns true if the components wants to show the channel (it has a related property) */
-		bool ComponentWantsToShowChannel() const;
-
-		/** Returns the starting address of the screen component */
-		int32 GetStartingChannel() const;
-
-		/** Returns the color of the widget */
-		FLinearColor GetColor() const;
-
-		/** Returns true if this Model handles the Other component */
-		bool Equals(UDMXPixelMappingBaseComponent* Other) const;
-
-	private:
-		/** Called when selected Components changed */
-		void OnSelectedComponentsChanged();
-
-		/** True if the component is selected */
-		bool bSelected = false;
-
-		/** The output component */
-		TWeakObjectPtr<UDMXPixelMappingScreenComponent> WeakScreenComponent;
-
-		/** The toolkit from which the Model currently sources */
-		TWeakPtr<FDMXPixelMappingToolkit> WeakToolkit;
-	};
 }
diff --git a/Engine/Plugins/VirtualProduction/DMX/DMXPixelMapping/Source/DMXPixelMappingEditor/Private/Views/SDMXPixelMappingDesignerView.cpp b/Engine/Plugins/VirtualProduction/DMX/DMXPixelMapping/Source/DMXPixelMappingEditor/Private/Views/SDMXPixelMappingDesignerView.cpp
--- a/Engine/Plugins/VirtualProduction/DMX/DMXPixelMapping/Source/DMXPixelMappingEditor/Private/Views/SDMXPixelMappingDesignerView.cpp
+++ b/Engine/Plugins/VirtualProduction/DMX/DMXPixelMapping/Source/DMXPixelMappingEditor/Private/Views/SDMXPixelMappingDesignerView.cpp
@@ -698,14 +698,8 @@ void SDMXPixelMappingDesignerView::RebuildDesigner()
 		CachedRendererComponent->ForEachChild([this, &Toolkit, &ComponentCanvas](UDMXPixelMappingBaseComponent* Component)
 			{
 				using namespace UE::DMX;
-				if (UDMXPixelMappingScreenComponent* ScreenComponent = Cast<UDMXPixelMappingScreenComponent>(Component))
-				{
-					const TSharedRef<SDMXPixelMappingScreenComponent> ScreenComponentWidget = SNew(SDMXPixelMappingScreenComponent, Toolkit.ToSharedRef(), ScreenComponent);
-					ScreenComponentWidget->AddToCanvas(ComponentCanvas);
 
-					OutputComponentWidgets.Add(ScreenComponentWidget);
-				}
-				else if (UDMXPixelMappingOutputComponent* OutputComponent = Cast<UDMXPixelMappingOutputComponent>(Component))
+				if (UDMXPixelMappingOutputComponent* OutputComponent = Cast<UDMXPixelMappingOutputComponent>(Component))
 				{
 					const TSharedRef<SDMXPixelMappingOutputComponent> ComponentWidget = SNew(SDMXPixelMappingOutputComponent, Toolkit.ToSharedRef(), OutputComponent);
 					ComponentWidget->AddToCanvas(ComponentCanvas);
diff --git a/Engine/Plugins/VirtualProduction/DMX/DMXPixelMapping/Source/DMXPixelMappingEditor/Private/Views/SDMXPixelMappingDetailsView.cpp b/Engine/Plugins/VirtualProduction/DMX/DMXPixelMapping/Source/DMXPixelMappingEditor/Private/Views/SDMXPixelMappingDetailsView.cpp
--- a/Engine/Plugins/VirtualProduction/DMX/DMXPixelMapping/Source/DMXPixelMappingEditor/Private/Views/SDMXPixelMappingDetailsView.cpp
+++ b/Engine/Plugins/VirtualProduction/DMX/DMXPixelMapping/Source/DMXPixelMappingEditor/Private/Views/SDMXPixelMappingDetailsView.cpp
@@ -3,21 +3,19 @@
 #include "Views/SDMXPixelMappingDetailsView.h"
 
 #include "ColorSpace/DMXPixelMappingColorSpace_RGBCMY.h"
-#include "Components/DMXPixelMappingRootComponent.h"
 #include "Components/DMXPixelMappingFixtureGroupComponent.h"
 #include "Components/DMXPixelMappingFixtureGroupItemComponent.h"
-#include "Components/DMXPixelMappingScreenComponent.h"
-#include "Components/DMXPixelMappingMatrixComponent.h"
 #include "Components/DMXPixelMappingMatrixCellComponent.h"
+#include "Components/DMXPixelMappingMatrixComponent.h"
 #include "Components/DMXPixelMappingOutputDMXComponent.h"
 #include "Components/DMXPixelMappingRendererComponent.h"
+#include "Components/DMXPixelMappingRootComponent.h"
 #include "Customizations/DMXPixelMappingColorSpaceDetails_RGBCMY.h"
 #include "Customizations/DMXPixelMappingDetailCustomization_FixtureGroup.h"
 #include "Customizations/DMXPixelMappingDetailCustomization_FixtureGroupItem.h"
-#include "Customizations/DMXPixelMappingDetailCustomization_Screen.h"
-#include "Customizations/DMXPixelMappingDetailCustomization_Renderer.h"
 #include "Customizations/DMXPixelMappingDetailCustomization_Matrix.h"
 #include "Customizations/DMXPixelMappingDetailCustomization_OutputDMX.h"
+#include "Customizations/DMXPixelMappingDetailCustomization_Renderer.h"
 #include "DetailsViewArgs.h"
 #include "DMXPixelMappingComponentReference.h"
 #include "IDetailsView.h"
@@ -26,7 +24,6 @@
 #include "PropertyEditorModule.h"
 #include "Toolkits/DMXPixelMappingToolkit.h"
 
-
 void SDMXPixelMappingDetailsView::Construct(const FArguments& InArgs, const TSharedPtr<FDMXPixelMappingToolkit>& InToolkit)
 {
 	WeakToolkit = InToolkit;
@@ -117,9 +114,6 @@ void SDMXPixelMappingDetailsView::RegisterCustomizations()
 
 	FOnGetDetailCustomizationInstance MatrixCustomizationInstance = FOnGetDetailCustomizationInstance::CreateStatic(&FDMXPixelMappingDetailCustomization_Matrix::MakeInstance, WeakToolkit);
 	DetailsView->RegisterInstancedCustomPropertyLayout(UDMXPixelMappingMatrixComponent::StaticClass(), MatrixCustomizationInstance);
-	
-	FOnGetDetailCustomizationInstance ScreenCustomizationInstance = FOnGetDetailCustomizationInstance::CreateStatic(&FDMXPixelMappingDetailCustomization_Screen::MakeInstance, WeakToolkit);
-	DetailsView->RegisterInstancedCustomPropertyLayout(UDMXPixelMappingScreenComponent::StaticClass(), ScreenCustomizationInstance);
 
 	FOnGetDetailCustomizationInstance ColorSpaceCustomizationInstance_RGBCMY = FOnGetDetailCustomizationInstance::CreateStatic(&FDMXPixelMappingColorSpaceDetails_RGBCMY::MakeInstance);
 	DetailsView->RegisterInstancedCustomPropertyLayout(UDMXPixelMappingColorSpace_RGBCMY::StaticClass(), ColorSpaceCustomizationInstance_RGBCMY);
diff --git a/Engine/Plugins/VirtualProduction/DMX/DMXPixelMapping/Source/DMXPixelMappingEditor/Private/Widgets/SDMXPixelMappingOutputComponent.cpp b/Engine/Plugins/VirtualProduction/DMX/DMXPixelMapping/Source/DMXPixelMappingEditor/Private/Widgets/SDMXPixelMappingOutputComponent.cpp
--- a/Engine/Plugins/VirtualProduction/DMX/DMXPixelMapping/Source/DMXPixelMappingEditor/Private/Widgets/SDMXPixelMappingOutputComponent.cpp
+++ b/Engine/Plugins/VirtualProduction/DMX/DMXPixelMapping/Source/DMXPixelMappingEditor/Private/Widgets/SDMXPixelMappingOutputComponent.cpp
@@ -189,381 +189,6 @@ namespace UE::DMX
 	
 		return FSlateRenderTransform(Quaternion, FVector2D::ZeroVector);
 	}
-
-	class SDMXPixelMappingScreenComponentCell
-		: public SCompoundWidget
-	{
-	public:
-		SLATE_BEGIN_ARGS(SDMXPixelMappingScreenComponentCell)
-			: _ChannelOffset(0)
-			, _StartingChannel(1)
-			, _Universe(1)
-		{}
-			/** The offset of this cell from the starting address */
-			SLATE_ARGUMENT(int32, ChannelOffset)
-
-			/** The starting address of the screen component */
-			SLATE_ATTRIBUTE(int32, StartingChannel)
-
-			/** The universe of the screen component */
-			SLATE_ATTRIBUTE(int32, Universe)
-
-		SLATE_END_ARGS()
-
-		/** Constructs this widget */
-		void Construct(const FArguments& InArgs, const TSharedRef<FDMXPixelMappingScreenComponentModel>& InModel)
-		{
-			Model = InModel;
-
-			ChannelOffset = InArgs._ChannelOffset;
-			Universe = InArgs._Universe;
-			StartingChannel = InArgs._StartingChannel;
-
-			ChildSlot
-				[
-					SAssignNew(ContentBorder, SBorder)
-					.BorderImage(FDMXPixelMappingEditorStyle::Get().GetBrush("DMXPixelMappingEditor.ComponentBorder"))
-					.HAlign(HAlign_Fill)
-					.VAlign(VAlign_Fill)
-				];
-
-			RebuildAddressesContent();
-
-			UDMXPixelMappingEditorSettings::OnEditorSettingsChanged.AddSP(this, &SDMXPixelMappingScreenComponentCell::RebuildAddressesContent);
-		}
-
-	private:
-		/** Rebuilds addresses content */
-		void RebuildAddressesContent()
-		{
-			const FDMXPixelMappingDesignerSettings& DesignerSettings = GetDefault<UDMXPixelMappingEditorSettings>()->DesignerSettings;
-			if (!DesignerSettings.bShowPatchInfo)
-			{
-				return;
-			}
-
-			if (!ensureMsgf(StartingChannel.IsSet() && Universe.IsSet(), TEXT("Trying to display universe and address of pixel mapping screen component, but universe and/or address are not set.")))
-			{
-				return;
-			}
-
-			ContentBorder->SetContent
-			(
-				SNew(SBorder)
-				.BorderImage(FAppStyle::GetBrush("NoBorder"))
-				.HAlign(HAlign_Fill)
-				.VAlign(VAlign_Fill)
-				.Padding(6.f)
-				[
-					SNew(SScaleBox)
-					.Stretch(EStretch::ScaleToFit)
-					.StretchDirection(EStretchDirection::DownOnly)
-					.VAlign(VAlign_Center)
-					.HAlign(HAlign_Center)
-					[
-						SNew(STextBlock)
-						.Font(FCoreStyle::GetDefaultFontStyle("Regular", 12))
-						.Text_Lambda([this]()
-							{
-								const int32 UniverseWithOffset = Universe.Get() + ChannelOffset / DMX_MAX_ADDRESS;
-								const int32 ChannelWithOffset = [this, UniverseWithOffset]()
-								{
-									if (UniverseWithOffset > 0)
-									{
-										return (StartingChannel.Get() + ChannelOffset) % DMX_MAX_ADDRESS;
-									}
-									return StartingChannel.Get() + ChannelOffset;
-								}();
-
-								if (Model->ComponentWantsToShowChannel() && Model->ComponentWantsToShowUniverse())
-								{
-									FString AddressString = FString::FromInt(UniverseWithOffset) + TEXT(".") + FString::FromInt(ChannelWithOffset);
-									return FText::FromString(AddressString);
-								}
-								else if (Model->ComponentWantsToShowUniverse())
-								{
-									FString UniverseString = FString::FromInt(UniverseWithOffset);
-									return FText::FromString(UniverseString);
-								}
-								else if (Model->ComponentWantsToShowChannel())
-								{
-									FString ChannelString = FString::FromInt(ChannelWithOffset);
-									return FText::FromString(ChannelString);
-								}
-								return FText::GetEmpty();
-							})
-					]
-				]
-			);
-		}
-
-		/** Border to hold contents of the widget */
-		TSharedPtr<SBorder> ContentBorder;
-
-		/** The model used by this widget */
-		TSharedPtr<FDMXPixelMappingScreenComponentModel> Model;
-
-		// Slate args
-		int32 ChannelOffset = 0;
-		TAttribute<int32> Universe;
-		TAttribute<int32> StartingChannel;
-	};
-
-	/** Widget to draw the inners of a screen component */
-	class SDMXPixelMappingScreenComponentGrid
-		: public SCompoundWidget
-	{
-	public:
-		SLATE_BEGIN_ARGS(SDMXPixelMappingScreenComponentGrid)
-		{}
-		SLATE_END_ARGS()
-
-		/** Constructs this widget */
-		void Construct(const FArguments& InArgs, const TSharedRef<FDMXPixelMappingScreenComponentModel>& InModel)
-		{
-			Model = InModel;
-
-			ChildSlot
-				.VAlign(VAlign_Fill)
-				.HAlign(HAlign_Fill)
-				[
-					SAssignNew(GridPanel, SUniformGridPanel)
-				];
-
-			CachedCellFormat = Model->GetCellFormat();
-			CachedDistribution = Model->GetDistribution();
-
-			RebuildGridContent();
-		}
-
-	protected:
-		//~ Begin SWidget Interface
-		virtual void Tick(const FGeometry& AllottedGeometry, const double InCurrentTime, const float InDeltaTime) override
-		{
-			if ((GridPanel->GetChildren()->Num() != Model->GetNumColumns() * Model->GetNumRows()) ||
-				CachedDistribution != Model->GetDistribution() ||
-				CachedCellFormat != Model->GetCellFormat())
-			{
-				// Handle cell format changes
-				CachedDistribution = Model->GetDistribution();
-				CachedCellFormat = Model->GetCellFormat();
-
-				RebuildGridContent();
-			}
-		}
-		//~ End SWidget Interface
-
-	private:
-		/** Rebuilds the grid content */
-		void RebuildGridContent()
-		{
-			GridPanel->ClearChildren();
-
-			TArray<int32> CellIndices;
-			const int32 TotalCells = Model->GetNumColumns() * Model->GetNumRows();
-			for (int32 CellID = 0; CellID < TotalCells; ++CellID)
-			{
-				CellIndices.Add(CellID);
-			}
-
-			TArray<int32> DistributionSortedCellIndices;
-			FDMXRuntimeUtils::PixelMappingDistributionSort(Model->GetDistribution(), Model->GetNumColumns(), Model->GetNumRows(), CellIndices, DistributionSortedCellIndices);
-
-			const uint32 NumChannelsPerCell = FDMXPixelMappingUtils::GetNumChannelsPerCell(Model->GetCellFormat());
-
-			uint32 XYIndex = 0;
-			for (int32 YIndex = 0; YIndex < Model->GetNumRows(); ++YIndex)
-			{
-				for (int32 XIndex = 0; XIndex < Model->GetNumColumns(); ++XIndex)
-				{
-					if (!ensureMsgf(DistributionSortedCellIndices.IsValidIndex(XYIndex), TEXT("Pixel mapping cell index out of range.")))
-					{
-						break;
-					}
-
-					const int32 ChannelOffset = DistributionSortedCellIndices[XYIndex] * NumChannelsPerCell;
-					GridPanel->AddSlot(XIndex, YIndex)
-						.HAlign(HAlign_Fill)
-						.VAlign(VAlign_Fill)
-						[
-							SNew(SDMXPixelMappingScreenComponentCell, Model.ToSharedRef())
-							.ChannelOffset(ChannelOffset)
-							.StartingChannel_Lambda([this]()
-								{
-									return Model->GetStartingChannel();
-								})
-							.Universe_Lambda([this]()
-								{
-									return Model->GetUniverse();
-								})
-						];
-
-					XYIndex++;
-				}
-			}
-		}
-
-		/** Cell format, cached to detect changes */
-		EDMXCellFormat CachedCellFormat = EDMXCellFormat::PF_RGB;
-
-		/** Pixelmapping distribution, cached to detect changes */
-		EDMXPixelMappingDistribution CachedDistribution = EDMXPixelMappingDistribution::TopLeftToRight;
-
-		/** The actual grid panel */
-		TSharedPtr<SUniformGridPanel> GridPanel;
-
-		/** The model used by this widget */
-		TSharedPtr<FDMXPixelMappingScreenComponentModel> Model;
-	};
-
-	/** Simple view of a screen component, useful if many cells are drawn */
-	class SDMXPixelMappingScreenComponentSimplistic
-		: public SCompoundWidget
-	{
-	public:
-		SLATE_BEGIN_ARGS(SDMXPixelMappingScreenComponentSimplistic)
-		{}
-		SLATE_END_ARGS()
-
-		//~ Begin SCompoundWidget interface
-		void Construct(const FArguments& InArgs, const TSharedRef<FDMXPixelMappingScreenComponentModel>& InModel)
-		{
-			ChildSlot
-			[
-				SNew(SBox)
-				.HAlign(HAlign_Fill)
-				.VAlign(VAlign_Fill)
-				[
-					SNew(SScaleBox)
-					.Stretch(EStretch::ScaleToFit)
-					.StretchDirection(EStretchDirection::DownOnly)
-					[
-						SNew(SVerticalBox)
-						+ SVerticalBox::Slot()
-						.HAlign(HAlign_Center)
-						.VAlign(VAlign_Center)
-						[
-							SNew(STextBlock)
-							.Font(FCoreStyle::GetDefaultFontStyle("Regular", 24))
-							.ColorAndOpacity(FSlateColor::UseForeground())
-							.Text_Lambda([InModel]()
-								{
-									return FText::Format(LOCTEXT("Num_Pixels", "{0} x {1} pixels"), InModel->GetNumColumns(), InModel->GetNumRows());
-								})
-						]
-						+ SVerticalBox::Slot()
-						.HAlign(HAlign_Center)
-						.VAlign(VAlign_Center)
-						[
-							SNew(STextBlock)
-							.Font(FCoreStyle::GetDefaultFontStyle("Regular", 24))
-							.ColorAndOpacity(FSlateColor::UseForeground())
-							.Text_Lambda([InModel]()
-								{
-									return FText::FromString(FString::Printf(TEXT("%i.%i"), InModel->GetUniverse(), InModel->GetStartingChannel()));
-								})
-						]
-					]
-				]
-			];
-		}
-		//~ Begin SCompoundWidget interface
-
-	private:
-		/** The model used by this widget */
-		TSharedPtr<FDMXPixelMappingScreenComponentModel> Model;
-	};
-
-	void SDMXPixelMappingScreenComponent::Construct(const FArguments& InArgs, const TSharedRef<FDMXPixelMappingToolkit>& InToolkit, TWeakObjectPtr<UDMXPixelMappingScreenComponent> ScreenComponent)
-	{
-		if (!ScreenComponent.IsValid())
-		{
-			return;
-		}
-
-		SetCanTick(true);
-
-		WeakToolkit = InToolkit;
-		Model = MakeShared<FDMXPixelMappingScreenComponentModel>(InToolkit, ScreenComponent);
-
-		ChildSlot
-		[
-			SNew(SBox)
-			.WidthOverride_Lambda([this]()
-				{
-					return Model->GetSize().X;
-				})
-			.HeightOverride_Lambda([this]()
-				{
-					return Model->GetSize().Y;
-				})
-			[
-				SNew(SOverlay)
-
-				+ SOverlay::Slot() 
-					[
-						SAssignNew(ContentBorder, SBorder)
-						.BorderImage(FAppStyle::GetBrush("NoBorder"))
-						.Padding(0.f)
-					]
-
-				+ SOverlay::Slot() // Required to correctly draw the color around the grid
-				[
-					SNew(SBorder)
-					.BorderImage(FDMXPixelMappingEditorStyle::Get().GetBrush("DMXPixelMappingEditor.ComponentBorder"))
-					.BorderBackgroundColor_Lambda([this]()
-						{
-							return Model->GetColor();
-						})
-				]
-			]
-		];	
-
-		UpdateContent();
-	}
-
-	bool SDMXPixelMappingScreenComponent::Equals(UDMXPixelMappingBaseComponent* Component) const
-	{
-		if (Component)
-		{
-			return Model->Equals(Component);
-		}
-		return false;
-	}
-
-	FVector2D SDMXPixelMappingScreenComponent::GetPosition() const
-	{
-		return Model->GetPosition();
-	}
-
-	void SDMXPixelMappingScreenComponent::Tick(const FGeometry& AllottedGeometry, const double InCurrentTime, const float InDeltaTime)
-	{
-		const bool bShouldDrawSimplisticView = Model->GetNumColumns() * Model->GetNumRows() > 160;
-		if (bDrawsSimplisticView != bShouldDrawSimplisticView)
-		{
-			UpdateContent();
-		}
-	}
-
-	void SDMXPixelMappingScreenComponent::UpdateContent()
-	{
-		if(Model->GetNumColumns() * Model->GetNumRows() > 160)
-		{
-			bDrawsSimplisticView = true;
-			ContentBorder->SetContent
-			(
-				SNew(SDMXPixelMappingScreenComponentSimplistic, Model.ToSharedRef())
-			);
-		}
-		else
-		{
-			bDrawsSimplisticView = false;
-			ContentBorder->SetContent
-			(
-				SNew(SDMXPixelMappingScreenComponentGrid, Model.ToSharedRef())
-			);
-		}
-	}
 }
 
 #undef LOCTEXT_NAMESPACE
diff --git a/Engine/Plugins/VirtualProduction/DMX/DMXPixelMapping/Source/DMXPixelMappingEditor/Private/Widgets/SDMXPixelMappingOutputComponent.h b/Engine/Plugins/VirtualProduction/DMX/DMXPixelMapping/Source/DMXPixelMappingEditor/Private/Widgets/SDMXPixelMappingOutputComponent.h
--- a/Engine/Plugins/VirtualProduction/DMX/DMXPixelMapping/Source/DMXPixelMappingEditor/Private/Widgets/SDMXPixelMappingOutputComponent.h
+++ b/Engine/Plugins/VirtualProduction/DMX/DMXPixelMapping/Source/DMXPixelMappingEditor/Private/Widgets/SDMXPixelMappingOutputComponent.h
@@ -12,13 +12,10 @@ class SBox;
 class STextBlock;
 class UDMXPixelMappingBaseComponent;
 class UDMXPixelMappingOutputComponent;
-class UDMXPixelMappingScreenComponent;
 
 namespace UE::DMX
 {
 	class FDMXPixelMappingOutputComponentModel;
-	class FDMXPixelMappingScreenComponentModel;
-
 
 	/** Interface for Output Component Widgets */
 	class IDMXPixelMappingOutputComponentWidgetInterface
@@ -52,7 +49,7 @@ namespace UE::DMX
 	};
 
 
-	/** Widget that draws an Output Component. For Screen Component, see SDMXPixelMappingScreenComponent */
+	/** Widget that draws an Output Component. */
 	class SDMXPixelMappingOutputComponent
 		: public IDMXPixelMappingOutputComponentWidgetInterface
 		, public SCompoundWidget
@@ -88,51 +85,4 @@ namespace UE::DMX
 		/** The toolkit that owns */
 		TWeakPtr<FDMXPixelMappingToolkit> WeakToolkit;
 	};
-
-	/** Widget that draws a Screen Component. While the screen component cannot be added to new assets anymore, we still draw it for old ones. */
-	class SDMXPixelMappingScreenComponent
-		: public IDMXPixelMappingOutputComponentWidgetInterface
-		, public SCompoundWidget
-	{
-	public:
-		SLATE_BEGIN_ARGS(SDMXPixelMappingScreenComponent)
-		{}
-		SLATE_END_ARGS()
-
-		/** Constructs the widget */
-		void Construct(const FArguments& InArgs, const TSharedRef<FDMXPixelMappingToolkit>& InToolkit, TWeakObjectPtr<UDMXPixelMappingScreenComponent> ScreenComponent);
-
-		//~ Begin IDMXPixelMappingOutputComponentWidgetInterface
-		virtual TSharedRef<SWidget> AsWidget() override { return AsShared(); };
-		virtual bool Equals(UDMXPixelMappingBaseComponent* Component) const override;
-	protected:
-		virtual FVector2D GetPosition() const override;
-		//~ End IDMXPixelMappingOutputComponentWidgetInterface
-
-		//~ Begin SWidget interface
-		virtual void Tick(const FGeometry& AllottedGeometry, const double InCurrentTime, const float InDeltaTime) override;
-		//~ End SWidget interface
-
-	private:
-		/** Updates the content of the widget */
-		void UpdateContent();
-
-		/** True if them model currently draws a simplistic view of the component */
-		bool bDrawsSimplisticView = false;
-
-		/** When added to a parent, the canvas it was added to */
-		TWeakPtr<SConstraintCanvas> ParentCanvas;
-
-		/** The canvas slot of the component widget */
-		SConstraintCanvas::FSlot* Slot = nullptr;
-
-		/** Border that holds the content */
-		TSharedPtr<SBorder> ContentBorder;
-
-		/** The model for this widget */
-		TSharedPtr<FDMXPixelMappingScreenComponentModel> Model;
-
-		/** The toolkit that owns */
-		TWeakPtr<FDMXPixelMappingToolkit> WeakToolkit;
-	};
 }
\ No newline at end of file
diff --git a/Engine/Plugins/VirtualProduction/DMX/DMXPixelMapping/Source/DMXPixelMappingRuntime/Private/Components/DMXPixelMappingRendererComponent.cpp b/Engine/Plugins/VirtualProduction/DMX/DMXPixelMapping/Source/DMXPixelMappingRuntime/Private/Components/DMXPixelMappingRendererComponent.cpp
--- a/Engine/Plugins/VirtualProduction/DMX/DMXPixelMapping/Source/DMXPixelMappingRuntime/Private/Components/DMXPixelMappingRendererComponent.cpp
+++ b/Engine/Plugins/VirtualProduction/DMX/DMXPixelMapping/Source/DMXPixelMappingRuntime/Private/Components/DMXPixelMappingRendererComponent.cpp
@@ -9,7 +9,6 @@
 #include "Components/DMXPixelMappingFixtureGroupItemComponent.h"
 #include "Components/DMXPixelMappingMatrixComponent.h"
 #include "Components/DMXPixelMappingRootComponent.h"
-#include "Components/DMXPixelMappingScreenComponent.h"
 #include "DMXPixelMapping.h"
 #include "DMXPixelMappingPixelMapRenderer.h"
 #include "DMXPixelMappingPreprocessRenderer.h"
@@ -496,24 +495,7 @@ void UDMXPixelMappingRendererComponent::RenderEditorPreviewTexture()
 	PixelPreviewParams.Reserve(DownsamplePixelCount_DEPRECATED);
 	
 	ForEachChild([this, &PixelPreviewParams](UDMXPixelMappingBaseComponent* InComponent) {
-		if(UDMXPixelMappingScreenComponent* ScreenComponent = Cast<UDMXPixelMappingScreenComponent>(InComponent))
-		{
-			const FVector2D SizePixel = ScreenComponent->GetScreenPixelSize();
-			const int32 DownsampleIndexStart = ScreenComponent->GetPixelDownsamplePositionRange().Key;
-			const int32 PositionX = ScreenComponent->GetPosition().X;
-			const int32 PositionY = ScreenComponent->GetPosition().Y;
-
-			ScreenComponent->ForEachPixel([this, &PixelPreviewParams, SizePixel, PositionX, PositionY, DownsampleIndexStart](const int32 InXYIndex, const int32 XIndex, const int32 YIndex)
-				{
-					FDMXPixelMappingDownsamplePixelPreviewParam PixelPreviewParam;
-					PixelPreviewParam.ScreenPixelSize = SizePixel;
-					PixelPreviewParam.ScreenPixelPosition = FVector2D(PositionX + SizePixel.X * XIndex, PositionY + SizePixel.Y * YIndex);
-					PixelPreviewParam.DownsamplePosition = GetPixelPosition(InXYIndex + DownsampleIndexStart);
-
-					PixelPreviewParams.Add(MoveTemp(PixelPreviewParam));
-				});
-		}
-		else if (UDMXPixelMappingOutputDMXComponent* Component = Cast<UDMXPixelMappingOutputDMXComponent>(InComponent))
+		if (UDMXPixelMappingOutputDMXComponent* Component = Cast<UDMXPixelMappingOutputDMXComponent>(InComponent))
 		{
 			FDMXPixelMappingDownsamplePixelPreviewParam PixelPreviewParam;
 			PixelPreviewParam.ScreenPixelSize = Component->GetSize();
@@ -743,13 +725,7 @@ int32 UDMXPixelMappingRendererComponent::GetTotalDownsamplePixelCount()
 	constexpr bool bIsRecursive = true;
 	ForEachChildOfClass<UDMXPixelMappingOutputComponent>([&](UDMXPixelMappingOutputComponent* InComponent)
 		{
-			// If that is screen component
-			if (UDMXPixelMappingScreenComponent* ScreenComponent = Cast<UDMXPixelMappingScreenComponent>(InComponent))
-			{
-				DownsamplePixelCount_DEPRECATED += (ScreenComponent->NumXCells * ScreenComponent->NumYCells);
-			}
-			// If that is single pixel component
-			else if (Cast<UDMXPixelMappingOutputDMXComponent>(InComponent))
+			if (Cast<UDMXPixelMappingOutputDMXComponent>(InComponent))
 			{
 				DownsamplePixelCount_DEPRECATED++;
 			}
diff --git a/Engine/Plugins/VirtualProduction/DMX/DMXPixelMapping/Source/DMXPixelMappingRuntime/Private/Components/DMXPixelMappingScreenComponent.cpp b/Engine/Plugins/VirtualProduction/DMX/DMXPixelMapping/Source/DMXPixelMappingRuntime/Private/Components/DMXPixelMappingScreenComponent.cpp
--- a/Engine/Plugins/VirtualProduction/DMX/DMXPixelMapping/Source/DMXPixelMappingRuntime/Private/Components/DMXPixelMappingScreenComponent.cpp
+++ b/Engine/Plugins/VirtualProduction/DMX/DMXPixelMapping/Source/DMXPixelMappingRuntime/Private/Components/DMXPixelMappingScreenComponent.cpp
@@ -23,11 +23,14 @@
 
 DECLARE_CYCLE_STAT(TEXT("Send Screen"), STAT_DMXPixelMaping_SendScreen, STATGROUP_DMXPIXELMAPPING);
 
-#define LOCTEXT_NAMESPACE "DMXPixelMappingScreenComponent"
+#define LOCTEXT_NAMESPACE "DEPRECATED_DMXPixelMappingScreenComponent"
 
-const FVector2D UDMXPixelMappingScreenComponent::MinGridSize = FVector2D(1.f);
+// DMXPixelMappingScreenComponent is fully deprecated, accept the deprecated implementation still makes use of it
+PRAGMA_DISABLE_DEPRECATION_WARNINGS 
 
-UDMXPixelMappingScreenComponent::UDMXPixelMappingScreenComponent()
+const FVector2D UDEPRECATED_DMXPixelMappingScreenComponent::MinGridSize = FVector2D(1.f);
+
+UDEPRECATED_DMXPixelMappingScreenComponent::UDEPRECATED_DMXPixelMappingScreenComponent()
 	: bSendToAllOutputPorts(true)
 {
 	SetSize(FVector2D(500.f, 500.f)); 
@@ -46,12 +49,12 @@ UDMXPixelMappingScreenComponent::UDMXPixelMappingScreenComponent()
 }
 
 #if WITH_EDITOR
-void UDMXPixelMappingScreenComponent::PostEditChangeChainProperty(FPropertyChangedChainEvent& PropertyChangedChainEvent)
+void UDEPRECATED_DMXPixelMappingScreenComponent::PostEditChangeChainProperty(FPropertyChangedChainEvent& PropertyChangedChainEvent)
 {
 	// Call the parent at the first place
 	Super::PostEditChangeChainProperty(PropertyChangedChainEvent);
 	
-	if (PropertyChangedChainEvent.GetPropertyName() == GET_MEMBER_NAME_CHECKED(UDMXPixelMappingScreenComponent, OutputPortReferences))
+	if (PropertyChangedChainEvent.GetPropertyName() == GET_MEMBER_NAME_CHECKED(UDEPRECATED_DMXPixelMappingScreenComponent, OutputPortReferences))
 	{
 		// Rebuild the set of ports
 		OutputPorts.Reset();
@@ -68,16 +71,14 @@ void UDMXPixelMappingScreenComponent::PostEditChangeChainProperty(FPropertyChang
 		}
 	}
 	
-
-	PRAGMA_DISABLE_DEPRECATION_WARNINGS
-	if (PropertyChangedChainEvent.GetPropertyName() == GET_MEMBER_NAME_CHECKED(UDMXPixelMappingScreenComponent, NumXCells) ||
-		PropertyChangedChainEvent.GetPropertyName() == GET_MEMBER_NAME_CHECKED(UDMXPixelMappingScreenComponent, NumYCells) ||
-		PropertyChangedChainEvent.GetPropertyName() == GET_MEMBER_NAME_CHECKED(UDMXPixelMappingScreenComponent, LocalUniverse) ||
-		PropertyChangedChainEvent.GetPropertyName() == GET_MEMBER_NAME_CHECKED(UDMXPixelMappingScreenComponent, StartAddress) ||
-		PropertyChangedChainEvent.GetPropertyName() == GET_MEMBER_NAME_CHECKED(UDMXPixelMappingScreenComponent, Distribution) ||
-		PropertyChangedChainEvent.GetPropertyName() == GET_MEMBER_NAME_CHECKED(UDMXPixelMappingScreenComponent, PixelFormat) ||
-		PropertyChangedChainEvent.GetPropertyName() == GET_MEMBER_NAME_CHECKED(UDMXPixelMappingScreenComponent, bShowAddresses) ||
-		PropertyChangedChainEvent.GetPropertyName() == GET_MEMBER_NAME_CHECKED(UDMXPixelMappingScreenComponent, bShowUniverse))
+	if (PropertyChangedChainEvent.GetPropertyName() == GET_MEMBER_NAME_CHECKED(UDEPRECATED_DMXPixelMappingScreenComponent, NumXCells) ||
+		PropertyChangedChainEvent.GetPropertyName() == GET_MEMBER_NAME_CHECKED(UDEPRECATED_DMXPixelMappingScreenComponent, NumYCells) ||
+		PropertyChangedChainEvent.GetPropertyName() == GET_MEMBER_NAME_CHECKED(UDEPRECATED_DMXPixelMappingScreenComponent, LocalUniverse) ||
+		PropertyChangedChainEvent.GetPropertyName() == GET_MEMBER_NAME_CHECKED(UDEPRECATED_DMXPixelMappingScreenComponent, StartAddress) ||
+		PropertyChangedChainEvent.GetPropertyName() == GET_MEMBER_NAME_CHECKED(UDEPRECATED_DMXPixelMappingScreenComponent, Distribution) ||
+		PropertyChangedChainEvent.GetPropertyName() == GET_MEMBER_NAME_CHECKED(UDEPRECATED_DMXPixelMappingScreenComponent, PixelFormat) ||
+		PropertyChangedChainEvent.GetPropertyName() == GET_MEMBER_NAME_CHECKED(UDEPRECATED_DMXPixelMappingScreenComponent, bShowAddresses) ||
+		PropertyChangedChainEvent.GetPropertyName() == GET_MEMBER_NAME_CHECKED(UDEPRECATED_DMXPixelMappingScreenComponent, bShowUniverse))
 	{
 		if (ScreenComponentBox.IsValid())
 		{
@@ -94,24 +95,23 @@ void UDMXPixelMappingScreenComponent::PostEditChangeChainProperty(FPropertyChang
 			ScreenComponentBox->RebuildGrid(GridParams);
 		}
 	}
-	PRAGMA_ENABLE_DEPRECATION_WARNINGS
 }
 #endif // WITH_EDITOR
 
 #if WITH_EDITOR
-const FText UDMXPixelMappingScreenComponent::GetPaletteCategory()
+const FText UDEPRECATED_DMXPixelMappingScreenComponent::GetPaletteCategory()
 {
 	return LOCTEXT("Common", "Common");
 }
 #endif // WITH_EDITOR
 
-const FName& UDMXPixelMappingScreenComponent::GetNamePrefix()
+const FName& UDEPRECATED_DMXPixelMappingScreenComponent::GetNamePrefix()
 {
 	static FName NamePrefix = TEXT("DEPRECATED DMX Screen");
 	return NamePrefix;
 }
 
-void UDMXPixelMappingScreenComponent::AddColorToSendBuffer(const FColor& InColor, TArray<uint8>& OutDMXSendBuffer)
+void UDEPRECATED_DMXPixelMappingScreenComponent::AddColorToSendBuffer(const FColor& InColor, TArray<uint8>& OutDMXSendBuffer)
 {
 	if (PixelFormat == EDMXCellFormat::PF_R)
 	{
@@ -209,15 +209,13 @@ void UDMXPixelMappingScreenComponent::AddColorToSendBuffer(const FColor& InColor
 	}
 }
 
-UDMXPixelMappingRendererComponent* UDMXPixelMappingScreenComponent::GetRendererComponent() const
+UDMXPixelMappingRendererComponent* UDEPRECATED_DMXPixelMappingScreenComponent::GetRendererComponent() const
 {
 	return Cast<UDMXPixelMappingRendererComponent>(GetParent());
 }
 
-void UDMXPixelMappingScreenComponent::SendDMX()
+void UDEPRECATED_DMXPixelMappingScreenComponent::SendDMX()
 {
-	PRAGMA_DISABLE_DEPRECATION_WARNINGS
-
 	SCOPE_CYCLE_COUNTER(STAT_DMXPixelMaping_SendScreen);
 
 	UDMXPixelMappingRendererComponent* RendererComponent = GetRendererComponent();
@@ -315,15 +313,10 @@ void UDMXPixelMappingScreenComponent::SendDMX()
 			SendDMXIndex++;
 		}
 	}
-
-	PRAGMA_ENABLE_DEPRECATION_WARNINGS
 }
 
-void UDMXPixelMappingScreenComponent::QueueDownsample()
+void UDEPRECATED_DMXPixelMappingScreenComponent::QueueDownsample()
 {
-	// DEPRECATED 5.3
-	PRAGMA_DISABLE_DEPRECATION_WARNINGS
-
 	// Queue pixels into the downsample rendering
 	UDMXPixelMappingRendererComponent* RendererComponent = GetRendererComponent();
 	if (!ensure(RendererComponent))
@@ -375,27 +368,25 @@ void UDMXPixelMappingScreenComponent::QueueDownsample()
 
 	// End of downsample index
 	PixelDownsamplePositionRange.Value = PixelDownsamplePositionRange.Key + IterationCount;
-
-	PRAGMA_ENABLE_DEPRECATION_WARNINGS
 }
 
-void UDMXPixelMappingScreenComponent::RenderWithInputAndSendDMX()
+void UDEPRECATED_DMXPixelMappingScreenComponent::RenderWithInputAndSendDMX()
 {
 	// DEPRECATED 5.3
 	RenderAndSendDMX();
 }
 
-bool UDMXPixelMappingScreenComponent::CanBeMovedTo(const UDMXPixelMappingBaseComponent* Component) const
+bool UDEPRECATED_DMXPixelMappingScreenComponent::CanBeMovedTo(const UDMXPixelMappingBaseComponent* Component) const
 {
 	return Component && Component->IsA<UDMXPixelMappingRendererComponent>();
 }
 
-const FVector2D UDMXPixelMappingScreenComponent::GetScreenPixelSize() const
+const FVector2D UDEPRECATED_DMXPixelMappingScreenComponent::GetScreenPixelSize() const
 {
 	return FVector2D(GetSize().X / NumXCells, GetSize().Y / NumYCells);
 }
 
-void UDMXPixelMappingScreenComponent::ForEachPixel(ForEachPixelCallback InCallback)
+void UDEPRECATED_DMXPixelMappingScreenComponent::ForEachPixel(ForEachPixelCallback InCallback)
 {
 	int32 IndexXY = 0;
 	for (int32 NumYIndex = 0; NumYIndex < NumYCells; ++NumYIndex)
@@ -408,4 +399,6 @@ void UDMXPixelMappingScreenComponent::ForEachPixel(ForEachPixelCallback InCallba
 	}
 }
 
+PRAGMA_ENABLE_DEPRECATION_WARNINGS
+
 #undef LOCTEXT_NAMESPACE
diff --git a/Engine/Plugins/VirtualProduction/DMX/DMXPixelMapping/Source/DMXPixelMappingRuntime/Public/Components/DMXPixelMappingOutputDMXComponent.h b/Engine/Plugins/VirtualProduction/DMX/DMXPixelMapping/Source/DMXPixelMappingRuntime/Public/Components/DMXPixelMappingOutputDMXComponent.h
--- a/Engine/Plugins/VirtualProduction/DMX/DMXPixelMapping/Source/DMXPixelMappingRuntime/Public/Components/DMXPixelMappingOutputDMXComponent.h
+++ b/Engine/Plugins/VirtualProduction/DMX/DMXPixelMapping/Source/DMXPixelMappingRuntime/Public/Components/DMXPixelMappingOutputDMXComponent.h
@@ -13,11 +13,7 @@ class UDMXModulator;
 class UDMXPixelMappingColorSpace;
 
 
-/**
- * Base class for components that contain a fixture patch to send DMX.
- * 
- * For legacy reasons also used by deprecated DMXPixelMappingScreenComponent
- */
+/** Base class for components that contain a fixture patch to send DMX. */
 UCLASS(Abstract)
 class DMXPIXELMAPPINGRUNTIME_API UDMXPixelMappingOutputDMXComponent
 	: public UDMXPixelMappingOutputComponent
diff --git a/Engine/Plugins/VirtualProduction/DMX/DMXPixelMapping/Source/DMXPixelMappingRuntime/Public/Components/DMXPixelMappingScreenComponent.h b/Engine/Plugins/VirtualProduction/DMX/DMXPixelMapping/Source/DMXPixelMappingRuntime/Public/Components/DMXPixelMappingScreenComponent.h
--- a/Engine/Plugins/VirtualProduction/DMX/DMXPixelMapping/Source/DMXPixelMappingRuntime/Public/Components/DMXPixelMappingScreenComponent.h
+++ b/Engine/Plugins/VirtualProduction/DMX/DMXPixelMapping/Source/DMXPixelMappingRuntime/Public/Components/DMXPixelMappingScreenComponent.h
@@ -19,20 +19,20 @@ class SDMXPixelMappingScreenComponentBox;
 class UTextureRenderTarget2D;
 class UDMXPixelMappingRendererComponent;
 
-/**
- * DMX Screen(Grid) rendering component
- */
-UCLASS()
-class DMXPIXELMAPPINGRUNTIME_API UDMXPixelMappingScreenComponent
+/** DMX Screen component. Fully deprecated as of 5.5. */
+class UE_DEPRECATED(5.5, "DMXScreenComponent is no longer supported in favor of a fully Fixture Patch based workflow.") UDMXPixelMappingScreenComponent;
+UCLASS(Deprecated, meta = (DeprecationMessage = "Deprecated 5.5. DMXScreenComponent is no longer supported in favor of a fully Fixture Patch based workflow."))
+class DMXPIXELMAPPINGRUNTIME_API UDEPRECATED_DMXPixelMappingScreenComponent
 	: public UDMXPixelMappingOutputDMXComponent
 {
 	GENERATED_BODY()
+
 public:
 	using ForEachPixelCallback = TFunctionRef<void(const int32 /* IndexXY */, const int32 /* IndexX */, const int32 /* IndexY */)>;
 
 public:
 	/** Default Constructor */
-	UDMXPixelMappingScreenComponent();
+	UDEPRECATED_DMXPixelMappingScreenComponent();
 
 	//~ Begin UObject implementation
 #if WITH_EDITOR
