DMX: Fix Pixel Mapping layouts no longer layout children of Fixture Groups correctly when a single group is selected

#rnx

diff --git a/Engine/Plugins/VirtualProduction/DMX/DMXPixelMapping/Source/DMXPixelMappingEditor/Private/ViewModels/DMXPixelMappingLayoutViewModel.cpp b/Engine/Plugins/VirtualProduction/DMX/DMXPixelMapping/Source/DMXPixelMappingEditor/Private/ViewModels/DMXPixelMappingLayoutViewModel.cpp
--- a/Engine/Plugins/VirtualProduction/DMX/DMXPixelMapping/Source/DMXPixelMappingEditor/Private/ViewModels/DMXPixelMappingLayoutViewModel.cpp
+++ b/Engine/Plugins/VirtualProduction/DMX/DMXPixelMapping/Source/DMXPixelMappingEditor/Private/ViewModels/DMXPixelMappingLayoutViewModel.cpp
@@ -57,7 +57,12 @@ EDMXPixelMappingLayoutViewModelMode UDMXPixelMappingLayoutViewModel::GetMode() c
 {
 	// If a Fixture Group is selected, layout its children.
 	// If a Matrix is selected, layout its children.
-	if (MatrixComponents.Num() == 1 && FixtureGroupComponents.IsEmpty())
+	// If more than one Fixture Group is selected, layout those.
+	if (FixtureGroupComponents.Num() == 1)
+	{
+		return EDMXPixelMappingLayoutViewModelMode::LayoutFixtureGroupComponentChildren;
+	}
+	else if (MatrixComponents.Num() == 1 && FixtureGroupComponents.IsEmpty())
 	{
 		return EDMXPixelMappingLayoutViewModelMode::LayoutMatrixComponentChildren;
 	}
