DMX - Pixel Mapping now draws its class icon and thumbnail correcly, Pixel Mapping now renders a preview texture as icon when saved.

diff --git a/Engine/Plugins/VirtualProduction/DMX/DMXPixelMapping/Source/DMXPixelMappingEditor/Private/DMXPixelMappingEditorModule.cpp b/Engine/Plugins/VirtualProduction/DMX/DMXPixelMapping/Source/DMXPixelMappingEditor/Private/DMXPixelMappingEditorModule.cpp
--- a/Engine/Plugins/VirtualProduction/DMX/DMXPixelMapping/Source/DMXPixelMappingEditor/Private/DMXPixelMappingEditorModule.cpp
+++ b/Engine/Plugins/VirtualProduction/DMX/DMXPixelMapping/Source/DMXPixelMappingEditor/Private/DMXPixelMappingEditorModule.cpp
@@ -26,8 +26,11 @@ void FDMXPixelMappingEditorModule::StartupModule()
 
 	FDMXPixelMappingEditorCommands::Register();
 
-	// Any attempt to use GEditor right now will fail as it hasn't been initialized yet. Waiting for post engine init resolves that.
-	FCoreDelegates::OnPostEngineInit.AddRaw(this, &FDMXPixelMappingEditorModule::OnPostEngineInit);
+	// Access the Pixel Mapping editor style so it gets instantiated and registered early.
+	// This will allow the engine to find related class icon and thumbnail.
+	FDMXPixelMappingEditorStyle::Get();
+
+	UThumbnailManager::Get().RegisterCustomRenderer(UDMXPixelMapping::StaticClass(), UDMXPixelMappingThumbnailRendering::StaticClass());
 }
 
 void FDMXPixelMappingEditorModule::ShutdownModule()
@@ -60,14 +63,6 @@ void FDMXPixelMappingEditorModule::RegisterAssetTypeAction(IAssetTools& AssetToo
 	CreatedAssetTypeActions.Add(Action);
 }
 
-void FDMXPixelMappingEditorModule::OnPostEngineInit()
-{
-	if (GIsEditor)
-	{
-		UThumbnailManager::Get().RegisterCustomRenderer(UDMXPixelMapping::StaticClass(), UDMXPixelMappingThumbnailRendering::StaticClass());
-	}
-}
-
 IMPLEMENT_MODULE(FDMXPixelMappingEditorModule, DMXPixelMappingEditor)
 
 #undef LOCTEXT_NAMESPACE
diff --git a/Engine/Plugins/VirtualProduction/DMX/DMXPixelMapping/Source/DMXPixelMappingEditor/Private/DMXPixelMappingEditorModule.h b/Engine/Plugins/VirtualProduction/DMX/DMXPixelMapping/Source/DMXPixelMappingEditor/Private/DMXPixelMappingEditorModule.h
--- a/Engine/Plugins/VirtualProduction/DMX/DMXPixelMapping/Source/DMXPixelMappingEditor/Private/DMXPixelMappingEditorModule.h
+++ b/Engine/Plugins/VirtualProduction/DMX/DMXPixelMapping/Source/DMXPixelMappingEditor/Private/DMXPixelMappingEditorModule.h
@@ -39,9 +39,6 @@ public:
 private:
 	void RegisterAssetTypeAction(IAssetTools& AssetTools, TSharedRef<IAssetTypeActions> Action);
 
-private:
-	void OnPostEngineInit();
-
 public:
 	static const FName DMXPixelMappingEditorAppIdentifier;
 
diff --git a/Engine/Plugins/VirtualProduction/DMX/DMXPixelMapping/Source/DMXPixelMappingEditor/Private/DMXPixelMappingThumbnailRendering.cpp b/Engine/Plugins/VirtualProduction/DMX/DMXPixelMapping/Source/DMXPixelMappingEditor/Private/DMXPixelMappingThumbnailRendering.cpp
--- a/Engine/Plugins/VirtualProduction/DMX/DMXPixelMapping/Source/DMXPixelMappingEditor/Private/DMXPixelMappingThumbnailRendering.cpp
+++ b/Engine/Plugins/VirtualProduction/DMX/DMXPixelMapping/Source/DMXPixelMappingEditor/Private/DMXPixelMappingThumbnailRendering.cpp
@@ -1,6 +1,7 @@
 // Copyright Epic Games, Inc. All Rights Reserved.
 
 #include "DMXPixelMappingThumbnailRendering.h"
+
 #include "CanvasTypes.h"
 #include "DMXPixelMapping.h"
 #include "Engine/Texture.h"
@@ -14,7 +15,7 @@ bool UDMXPixelMappingThumbnailRendering::CanVisualizeAsset(UObject* Object)
 void UDMXPixelMappingThumbnailRendering::GetThumbnailSize(UObject* Object, float Zoom, uint32& OutWidth, uint32& OutHeight) const
 {
 	UTexture* ObjectTexture = GetThumbnailTextureFromObject(Object);
-	if (ObjectTexture != nullptr && ObjectTexture->GetResource() != nullptr)
+	if (ObjectTexture && ObjectTexture->GetResource())
 	{
 		OutWidth = Zoom * ObjectTexture->GetResource()->GetSizeX();
 		OutHeight = Zoom * ObjectTexture->GetResource()->GetSizeY();
@@ -26,21 +27,17 @@ void UDMXPixelMappingThumbnailRendering::GetThumbnailSize(UObject* Object, float
 	}
 }
 
-void UDMXPixelMappingThumbnailRendering::Draw(UObject* Object, int32 X, int32 Y, uint32 Width, uint32 Height, FRenderTarget*, FCanvas* Canvas, bool bAdditionalViewFamily)
+void UDMXPixelMappingThumbnailRendering::Draw(UObject* Object, int32 X, int32 Y, uint32 Width, uint32 Height, FRenderTarget* RenderTarget, FCanvas* Canvas, bool bAdditionalViewFamily)
 {
-	UTexture* ObjectTexture = GetThumbnailTextureFromObject(Object);
-	if (ObjectTexture != nullptr)
+	if (UTexture* Texture = GetThumbnailTextureFromObject(Object))
 	{
-		Canvas->DrawTile(X, Y, Width, Height, 0.0f, 0.0f, 1.0f, 1.0f, FLinearColor::White, ObjectTexture->GetResource(), false);
+		Texture->WaitForPendingInitOrStreaming();
+		Canvas->DrawTile(X, Y, Width, Height, 0.0f, 0.0f, 1.0f, 1.0f, FLinearColor::White, Texture->GetResource(), false);
 	}
 }
 
 UTexture* UDMXPixelMappingThumbnailRendering::GetThumbnailTextureFromObject(UObject* Object) const
 {
 	UDMXPixelMapping* DMXPixelMapping = Cast<UDMXPixelMapping>(Object);
-	if (DMXPixelMapping != nullptr)
-	{
-		return DMXPixelMapping->ThumbnailImage;
-	}
-	return nullptr;
+	return DMXPixelMapping ? DMXPixelMapping->ThumbnailImage : nullptr;
 }
diff --git a/Engine/Plugins/VirtualProduction/DMX/DMXPixelMapping/Source/DMXPixelMappingEditor/Private/Factories/DMXPixelMappingFactoryNew.cpp b/Engine/Plugins/VirtualProduction/DMX/DMXPixelMapping/Source/DMXPixelMappingEditor/Private/Factories/DMXPixelMappingFactoryNew.cpp
--- a/Engine/Plugins/VirtualProduction/DMX/DMXPixelMapping/Source/DMXPixelMappingEditor/Private/Factories/DMXPixelMappingFactoryNew.cpp
+++ b/Engine/Plugins/VirtualProduction/DMX/DMXPixelMapping/Source/DMXPixelMappingEditor/Private/Factories/DMXPixelMappingFactoryNew.cpp
@@ -23,11 +23,6 @@ UDMXPixelMappingFactoryNew::UDMXPixelMappingFactoryNew()
 	bEditAfterNew = true;
 }
 
-FName UDMXPixelMappingFactoryNew::GetNewAssetThumbnailOverride() const
-{
-	return TEXT("ClassThumbnail.DMXPixelMapping");
-}
-
 UObject* UDMXPixelMappingFactoryNew::FactoryCreateNew(UClass* Class, UObject* InParent, FName Name, EObjectFlags Flags, UObject* Context, FFeedbackContext* Warn)
 {
 	check(Class->IsChildOf(UDMXPixelMapping::StaticClass()));
diff --git a/Engine/Plugins/VirtualProduction/DMX/DMXPixelMapping/Source/DMXPixelMappingEditor/Private/Factories/DMXPixelMappingFactoryNew.h b/Engine/Plugins/VirtualProduction/DMX/DMXPixelMapping/Source/DMXPixelMappingEditor/Private/Factories/DMXPixelMappingFactoryNew.h
--- a/Engine/Plugins/VirtualProduction/DMX/DMXPixelMapping/Source/DMXPixelMappingEditor/Private/Factories/DMXPixelMappingFactoryNew.h
+++ b/Engine/Plugins/VirtualProduction/DMX/DMXPixelMapping/Source/DMXPixelMappingEditor/Private/Factories/DMXPixelMappingFactoryNew.h
@@ -16,7 +16,6 @@ public:
 
 private:
 	//~ Begin UFactory Interface
-	virtual FName GetNewAssetThumbnailOverride() const override;
 	virtual UObject* FactoryCreateNew(UClass* Class, UObject* InParent, FName Name, EObjectFlags Flags, UObject* Context, FFeedbackContext* Warn) override;
 	virtual FText GetDisplayName() const override;
 	virtual uint32 GetMenuCategories() const override;
diff --git a/Engine/Plugins/VirtualProduction/DMX/DMXPixelMapping/Source/DMXPixelMappingEditor/Private/Toolkits/DMXPixelMappingToolkit.cpp b/Engine/Plugins/VirtualProduction/DMX/DMXPixelMapping/Source/DMXPixelMappingEditor/Private/Toolkits/DMXPixelMappingToolkit.cpp
--- a/Engine/Plugins/VirtualProduction/DMX/DMXPixelMapping/Source/DMXPixelMappingEditor/Private/Toolkits/DMXPixelMappingToolkit.cpp
+++ b/Engine/Plugins/VirtualProduction/DMX/DMXPixelMapping/Source/DMXPixelMappingEditor/Private/Toolkits/DMXPixelMappingToolkit.cpp
@@ -4,12 +4,14 @@
 
 #include "Algo/Sort.h"
 #include "Algo/Transform.h"
+#include "CanvasItem.h"
+#include "CanvasTypes.h"
 #include "Components/DMXPixelMappingFixtureGroupComponent.h"
 #include "Components/DMXPixelMappingFixtureGroupItemComponent.h"
-#include "Components/DMXPixelMappingRendererComponent.h"
-#include "Components/DMXPixelMappingRootComponent.h"
 #include "Components/DMXPixelMappingMatrixCellComponent.h"
 #include "Components/DMXPixelMappingMatrixComponent.h"
+#include "Components/DMXPixelMappingRendererComponent.h"
+#include "Components/DMXPixelMappingRootComponent.h"
 #include "DMXPixelMapping.h"
 #include "DMXPixelMappingEditorCommands.h"
 #include "DMXPixelMappingEditorModule.h"
@@ -17,8 +19,11 @@
 #include "DMXPixelMappingEditorUtils.h"
 #include "DMXPixelMappingToolbar.h"
 #include "Engine/Texture.h"
+#include "Engine/Texture2D.h"
+#include "Engine/TextureRenderTarget2D.h"
 #include "Framework/Commands/GenericCommands.h"
 #include "Framework/MultiBox/MultiBoxExtender.h"
+#include "HitProxies.h"
 #include "K2Node_PixelMappingBaseComponent.h"
 #include "Library/DMXEntityFixturePatch.h"
 #include "Library/DMXLibrary.h"
@@ -27,6 +32,8 @@
 #include "ScopedTransaction.h"
 #include "Settings/DMXPixelMappingEditorSettings.h"
 #include "Templates/DMXPixelMappingComponentTemplate.h"
+#include "TextureResource.h"
+#include "UObject/ObjectSaveContext.h"
 #include "UObject/UObjectIterator.h"
 #include "Views/SDMXPixelMappingDesignerView.h"
 #include "Views/SDMXPixelMappingDetailsView.h"
@@ -36,7 +43,6 @@
 #include "Views/SDMXPixelMappingPreviewView.h"
 #include "Widgets/Docking/SDockTab.h"
 
-
 #define LOCTEXT_NAMESPACE "DMXPixelMappingToolkit"
 
 const FName FDMXPixelMappingToolkit::DMXLibraryViewTabID(TEXT("DMXPixelMappingEditor_DMXLibraryViewTabID"));
@@ -54,11 +60,6 @@ FDMXPixelMappingToolkit::FDMXPixelMappingToolkit()
 	Selection = NewObject<UDMXPixelMappingToolkitSelection>(GetTransientPackage(), NAME_None, RF_Transactional);
 }
 
-FDMXPixelMappingToolkit::~FDMXPixelMappingToolkit()
-{
-	SaveThumbnailImage();
-}
-
 void FDMXPixelMappingToolkit::InitPixelMappingEditor(const EToolkitMode::Type Mode, const TSharedPtr<class IToolkitHost>& InitToolkitHost, UDMXPixelMapping* InDMXPixelMapping)
 {
 	check(InDMXPixelMapping);
@@ -180,6 +181,9 @@ void FDMXPixelMappingToolkit::InitPixelMappingEditor(const EToolkitMode::Type Mo
 	// Set the scale children with parent property on the pixel mapping object, so it is accessible the runtime module.
 	const UDMXPixelMappingEditorSettings* EditorSettings = GetDefault<UDMXPixelMappingEditorSettings>();
 	InDMXPixelMapping->bEditorScaleChildrenWithParent = EditorSettings->DesignerSettings.bScaleChildrenWithParent;
+
+	// Listen to packages being saved
+	UPackage::PreSavePackageWithContextEvent.AddSP(this, &FDMXPixelMappingToolkit::PreSavePackage);
 }
 
 void FDMXPixelMappingToolkit::RegisterTabSpawners(const TSharedRef<class FTabManager>& InTabManager)
@@ -487,12 +491,70 @@ void FDMXPixelMappingToolkit::UpdateBlueprintNodes() const
 void FDMXPixelMappingToolkit::SaveThumbnailImage()
 {
 	UDMXPixelMapping* PixelMapping = GetDMXPixelMapping();
-	UDMXPixelMappingRendererComponent* ActiveRendererComponent = GetActiveRendererComponent();
+	UDMXPixelMappingRendererComponent* RendererComponent = GetActiveRendererComponent();
+	if (!PixelMapping || !RendererComponent)
+	{
+		return;
+	}
+
+	// Fully load the input texture
+	UTexture* InputTexture = RendererComponent ? RendererComponent->GetRenderedInputTexture() : nullptr;
+	if (InputTexture)
+	{
+		InputTexture->WaitForPendingInitOrStreaming();
+	}
+
+	// Don't set a thumbnail if no texture is available or no pixel mapping is setup
+	using namespace UE::DMXPixelMapping::Rendering;
+	const TArray<TSharedRef<FPixelMapRenderElement>> RenderElements = RendererComponent->GetPixelMapRenderElements();
+
+	const bool bIsEmptyMapping = !InputTexture || !InputTexture->GetResource() || RenderElements.IsEmpty();
+	if (bIsEmptyMapping)
+	{
+		PixelMapping->ThumbnailImage = nullptr;
+		return;
+	}
 
-	UTexture* Texture = ActiveRendererComponent ? ActiveRendererComponent->GetRenderedInputTexture() : nullptr;
-	if (IsValid(Texture) && Texture->IsFullyStreamedIn())
+	// Paint a preview of the pixel mapping
+	UTextureRenderTarget2D* RenderTarget = NewObject<UTextureRenderTarget2D>();
+
+	constexpr uint32 ThumbnailSize = 64;
+	RenderTarget->InitAutoFormat(ThumbnailSize, ThumbnailSize);
+
+	constexpr FHitProxyConsumer* HitProxyConsumer = nullptr;
+	FCanvas Canvas(RenderTarget->GameThread_GetRenderTargetResource(), HitProxyConsumer, FGameTime(), GMaxRHIFeatureLevel);
+	Canvas.Clear(FColor::Black);
+
+	if (!InputTexture || !InputTexture->GetResource())
+	{
+		return;
+	}
+
+	for (const TSharedRef<FPixelMapRenderElement>& Element : RenderElements)
+	{
+		const FVector2D UV = Element->GetParameters().UV;
+		const FVector2D UVSize = Element->GetParameters().UVSize;
+
+		constexpr uint32 Margin = 12;
+		constexpr uint32 ThumbnailSizeWithoutMargin = ThumbnailSize - Margin * 2.f;
+		const FVector2D NormalizedMargin = FVector2D(Margin) / FVector2D(ThumbnailSize);
+
+		const FVector2D Position = FVector2D(Margin) + UV * FIntPoint(ThumbnailSizeWithoutMargin, ThumbnailSizeWithoutMargin);
+		const FVector2D Size = UVSize * FVector2D(ThumbnailSizeWithoutMargin, ThumbnailSizeWithoutMargin);
+
+		FCanvasTileItem TileItem(Position, Size, Element->GetColor());
+		TileItem.BlendMode = ESimpleElementBlendMode::SE_BLEND_MAX;
+		TileItem.PivotPoint = FVector2D(0.5, 0.5);
+		TileItem.Rotation = FRotator(0.0, Element->GetParameters().Rotation, 0.0);
+		Canvas.DrawItem(TileItem);
+	}
+	Canvas.Flush_GameThread();
+
+	// Set the rendered thumbnail image 
+	if (IsValid(RenderTarget))
 	{
-		PixelMapping->ThumbnailImage = Texture;
+		PixelMapping->ThumbnailImage = NewObject<UTexture2D>(PixelMapping);
+		RenderTarget->UpdateTexture(PixelMapping->ThumbnailImage);
 	}
 }
 
@@ -841,6 +903,18 @@ void FDMXPixelMappingToolkit::CreateInternalViews()
 	GetOrCreateLayoutView();
 }
 
+void FDMXPixelMappingToolkit::PreSavePackage(class UPackage* Package, FObjectPreSaveContext Context)
+{
+	if (!Context.IsCooking())
+	{
+		UDMXPixelMapping* PixelMapping = GetDMXPixelMapping();
+		if (PixelMapping && PixelMapping->GetPackage() == Package)
+		{
+			SaveThumbnailImage();
+		}
+	}
+}
+
 void FDMXPixelMappingToolkit::RenameComponent(const FName& CurrentObjectName, const FString& DesiredObjectName) const
 {
 	const UDMXPixelMapping* PixelMapping = GetDMXPixelMapping();
diff --git a/Engine/Plugins/VirtualProduction/DMX/DMXPixelMapping/Source/DMXPixelMappingEditor/Private/Toolkits/DMXPixelMappingToolkit.h b/Engine/Plugins/VirtualProduction/DMX/DMXPixelMapping/Source/DMXPixelMappingEditor/Private/Toolkits/DMXPixelMappingToolkit.h
--- a/Engine/Plugins/VirtualProduction/DMX/DMXPixelMapping/Source/DMXPixelMappingEditor/Private/Toolkits/DMXPixelMappingToolkit.h
+++ b/Engine/Plugins/VirtualProduction/DMX/DMXPixelMapping/Source/DMXPixelMappingEditor/Private/Toolkits/DMXPixelMappingToolkit.h
@@ -76,11 +76,6 @@ public:
 	/** Default constructor */
 	FDMXPixelMappingToolkit();
 
-	/**
-	 * Destructor.
-	 */
-	virtual ~FDMXPixelMappingToolkit();
-
 	/**
 	 * Edits the specified Texture object.
 	 *
@@ -249,6 +244,9 @@ private:
 
 	void CreateInternalViews();
 
+	/** Called before any package was saved */
+	void PreSavePackage(class UPackage* Package, FObjectPreSaveContext Context);
+
 	/** Returns the check box state for the compared reset DMX mode */
 	ECheckBoxState GetEditorResetDMXModeCheckboxState(EDMXPixelMappingResetDMXMode CompareMode) const;
 
