DMX: Fix an issue where OSX did not render DMX Pixel Mappings correctly and sent zero instead of color values

diff --git a/Engine/Plugins/VirtualProduction/DMX/DMXPixelMapping/Source/DMXPixelMappingRenderer/Private/DMXPixelMappingRenderPixelMapProxy.cpp b/Engine/Plugins/VirtualProduction/DMX/DMXPixelMapping/Source/DMXPixelMappingRenderer/Private/DMXPixelMappingRenderPixelMapProxy.cpp
--- a/Engine/Plugins/VirtualProduction/DMX/DMXPixelMapping/Source/DMXPixelMappingRenderer/Private/DMXPixelMappingRenderPixelMapProxy.cpp
+++ b/Engine/Plugins/VirtualProduction/DMX/DMXPixelMapping/Source/DMXPixelMappingRenderer/Private/DMXPixelMappingRenderPixelMapProxy.cpp
@@ -127,7 +127,7 @@ namespace UE::DMXPixelMapping::Rendering::Private
 
 						const int32 RenderTargetPoisitionX = PixelIndex % OutputTextureSize.X;
 						const int32 RenderTargetPositionY = PixelIndex / OutputTextureSize.X;
-						
+
 						// Create shader permutations
 						FDMXPixelMappingRendererPS::FPermutationDomain PermutationVector;
 						PermutationVector.Set<FDMXPixelBlendingQualityDimension>(Element->GetParameters().CellBlendingQuality);
@@ -143,7 +143,7 @@ namespace UE::DMXPixelMapping::Rendering::Private
 						RHICmdList.ApplyCachedRenderTargets(GraphicsPSOInit);
 						GraphicsPSOInit.BlendState = TStaticBlendState<>::GetRHI();
 						GraphicsPSOInit.RasterizerState = TStaticRasterizerState<>::GetRHI();
-						GraphicsPSOInit.DepthStencilState = TStaticDepthStencilState<false, CF_Never>::GetRHI();
+						GraphicsPSOInit.DepthStencilState = TStaticDepthStencilState<false, CF_Always>::GetRHI();
 						GraphicsPSOInit.BoundShaderState.VertexDeclarationRHI = GFilterVertexDeclaration.VertexDeclarationRHI;
 						GraphicsPSOInit.BoundShaderState.VertexShaderRHI = VertexShader.GetVertexShader();
 						GraphicsPSOInit.BoundShaderState.PixelShaderRHI = PixelShader.GetPixelShader();
