DMX: Add gamma options for Pixel Mapping's DMX output

* Color space RGB now has three gamma options, linear, custom and same as color space
* Color space xyY and XYZ have an advanced option to set a custom gamma

diff --git a/Engine/Plugins/VirtualProduction/DMX/DMXPixelMapping/Source/DMXPixelMappingRuntime/Internal/ColorSpace/DMXPixelMappingColorSpace_RGBCMY.h b/Engine/Plugins/VirtualProduction/DMX/DMXPixelMapping/Source/DMXPixelMappingRuntime/Internal/ColorSpace/DMXPixelMappingColorSpace_RGBCMY.h
--- a/Engine/Plugins/VirtualProduction/DMX/DMXPixelMapping/Source/DMXPixelMappingRuntime/Internal/ColorSpace/DMXPixelMappingColorSpace_RGBCMY.h
+++ b/Engine/Plugins/VirtualProduction/DMX/DMXPixelMapping/Source/DMXPixelMappingRuntime/Internal/ColorSpace/DMXPixelMappingColorSpace_RGBCMY.h
@@ -4,13 +4,12 @@
 
 #include "DMXAttribute.h"
 #include "ColorSpace/DMXPixelMappingColorSpace.h"
-
 #include "ColorManagement/ColorSpace.h"
 #include "Templates/UniquePtr.h"
 
 #include "DMXPixelMappingColorSpace_RGBCMY.generated.h"
 
-
+/** The color space of the DMX output */
 UENUM(BlueprintType)
 enum class EDMXPixelMappingOutputColorSpace_RGBCMY : uint8
 {
@@ -21,6 +20,16 @@ enum class EDMXPixelMappingOutputColorSpace_RGBCMY : uint8
 	Plasa UMETA(DisplayName = "PLASA RGB - ANSI E1.54")
 };
 
+/** The gamma mode of the DMX output */
+UENUM(BlueprintType)
+enum class EDMXPixelMappingGamma_RGBCMY : uint8
+{
+	Linear,
+	AsOutputColorSpace,
+	Custom
+};
+
+/** Defines how the luminance channel (typically Dimmer) is generated */
 UENUM(BlueprintType)
 enum class EDMXPixelMappingLuminanceType_RGBCMY : uint8
 {
@@ -45,9 +54,17 @@ public:
 	//~ End DMXPixelMappingColorSpace interface
 
 	/** The color space to use */
-	UPROPERTY(EditAnywhere, Category = "Color Space", Meta = (DisplayPriority = 1, DisplayName = "Output Color Space"))
+	UPROPERTY(EditAnywhere, Category = "Color Space", Meta = (DisplayName = "Output Color Space"))
 	EDMXPixelMappingOutputColorSpace_RGBCMY PixelMappingOutputColorSpace = EDMXPixelMappingOutputColorSpace_RGBCMY::sRGB;
 
+	/** Applies Gamma correction to the output DMX Values */
+	UPROPERTY(EditAnywhere, Category = "Color Space")
+	EDMXPixelMappingGamma_RGBCMY OutputGamma = EDMXPixelMappingGamma_RGBCMY::AsOutputColorSpace;
+
+	/** The custom Gamma, whereas OutColor = Pow(InColor, 1 / CustomGamma) */
+	UPROPERTY(EditAnywhere, Category = "Color Space", Meta = (UIMin = 0.0, UIMax = 4.0, EditCondition = "OutputGamma == EDMXPixelMappingGamma_RGBCMY::Custom", EditConditionHides))
+	float CustomGamma = 2.2f;
+
 	/** If set, converts Red to Cyan */
 	UPROPERTY(EditAnywhere, Category = "RGB")
 	bool bSendCyan = false;
@@ -94,6 +111,7 @@ public:
 
 protected:
 	//~ Begin UObject Interface
+	virtual void Serialize(FArchive& Ar) override;
 	virtual void PostLoad() override;
 #if WITH_EDITOR
 	virtual void PostEditChangeChainProperty(FPropertyChangedChainEvent& PropertyChangedEvent) override;
@@ -104,6 +122,15 @@ private:
 	/** Updates the ColorSpace and ColorSpaceTransform members */
 	void UpdateColorSpaceAndTransform();
 
+	/** Applies gamma correction according to the color space to the Color. */
+	void EncodeGammaCorrection(EDMXPixelMappingOutputColorSpace_RGBCMY InPixelMappingOutputColorSpaceEnum, FLinearColor& InOutColor);
+
+	/** Applies a custom gamma correction to the Color. OutColor = Pow(InColor, CustomGamma) */
+	void EncodeGammaCorrection(float CustomGamma, FLinearColor& InOutColor);
+
+	/** Encodes Rec2020 in a linear color */
+	float EncodeRec2020(float InValue) const;
+
 	/** Gets the EColorSpace enum given the current Output Color Space Enum */
 	UE::Color::EColorSpace ConvertToOutputColorSpaceEnum(EDMXPixelMappingOutputColorSpace_RGBCMY InPixelMappingOutputColorSpaceEnum) const;
 
diff --git a/Engine/Plugins/VirtualProduction/DMX/DMXPixelMapping/Source/DMXPixelMappingRuntime/Internal/ColorSpace/DMXPixelMappingColorSpace_XYZ.h b/Engine/Plugins/VirtualProduction/DMX/DMXPixelMapping/Source/DMXPixelMappingRuntime/Internal/ColorSpace/DMXPixelMappingColorSpace_XYZ.h
--- a/Engine/Plugins/VirtualProduction/DMX/DMXPixelMapping/Source/DMXPixelMappingRuntime/Internal/ColorSpace/DMXPixelMappingColorSpace_XYZ.h
+++ b/Engine/Plugins/VirtualProduction/DMX/DMXPixelMapping/Source/DMXPixelMappingRuntime/Internal/ColorSpace/DMXPixelMappingColorSpace_XYZ.h
@@ -25,6 +25,13 @@ public:
 	virtual void SetRGBA(const FLinearColor& InColor) override;
 	//~ End DMXPixelMappingColorSpace interface
 
+	/** 
+	 * Output gamma of the Luminance, computed in xyY space whereas Y = Pow(Y, 1 / OutputGamma).
+	 * CIE 1931 XYZ is linear gamma so typically no gamma should be applied.
+	 */
+	UPROPERTY(EditAnywhere, AdvancedDisplay, Category = "Color Space")
+	float CustomGamma = 1.f;
+
 	/** Attribute sent for X */
 	UPROPERTY(EditAnywhere, Category = "XYZ", Meta = (DisplayName = "X Attribute"))
 	FDMXAttributeName XAttribute;
diff --git a/Engine/Plugins/VirtualProduction/DMX/DMXPixelMapping/Source/DMXPixelMappingRuntime/Internal/ColorSpace/DMXPixelMappingColorSpace_xyY.h b/Engine/Plugins/VirtualProduction/DMX/DMXPixelMapping/Source/DMXPixelMappingRuntime/Internal/ColorSpace/DMXPixelMappingColorSpace_xyY.h
--- a/Engine/Plugins/VirtualProduction/DMX/DMXPixelMapping/Source/DMXPixelMappingRuntime/Internal/ColorSpace/DMXPixelMappingColorSpace_xyY.h
+++ b/Engine/Plugins/VirtualProduction/DMX/DMXPixelMapping/Source/DMXPixelMappingRuntime/Internal/ColorSpace/DMXPixelMappingColorSpace_xyY.h
@@ -24,6 +24,13 @@ public:
 	virtual void SetRGBA(const FLinearColor& InColor) override;
 	//~ End DMXPixelMappingColorSpace interface
 
+	/**
+	 * Output gamma of the Y component in xyY space whereas Y = Pow(Y, 1 / CustomGamma).
+	 * CIE 1931 xyY is linear gamma so typically no gamma should be applied.
+	 */
+	UPROPERTY(EditAnywhere, AdvancedDisplay, Category = "Color Space")
+	float CustomGamma = 1.f;
+
 	/** Attribute sent for x */
 	UPROPERTY(EditAnywhere, Category = "XY", Meta = (DisplayName = "x Attribute"))
 	FDMXAttributeName XAttribute;
@@ -32,7 +39,11 @@ public:
 	UPROPERTY(EditAnywhere, Category = "XY", Meta = (DisplayName = "y Attribute"))
 	FDMXAttributeName YAttribute;
 
-	/** Sets the range of the xyY color space. A value of 1.0 corresponds to the CIE 1931 xyY . */
+	/** 
+	 * Sets the range of the xyY color space. 
+	 * A value of 1.0 corresponds to CIE 1931 xyY. 
+	 * Many hardware lights use a more narrow range, typically 0.8. 
+	 */
 	UPROPERTY(EditAnywhere, Category = "XY", Meta = (ClampMin = 0.1, ClampMax = 2.0, UIMin = 0.1, UIMax = 1.0, DisplayName = "Color Space Range"))
 	float ColorSpaceRange = .8f;
 
diff --git a/Engine/Plugins/VirtualProduction/DMX/DMXPixelMapping/Source/DMXPixelMappingRuntime/Private/ColorSpace/DMXPixelMappingColorSpace_RGBCMY.cpp b/Engine/Plugins/VirtualProduction/DMX/DMXPixelMapping/Source/DMXPixelMappingRuntime/Private/ColorSpace/DMXPixelMappingColorSpace_RGBCMY.cpp
--- a/Engine/Plugins/VirtualProduction/DMX/DMXPixelMapping/Source/DMXPixelMappingRuntime/Private/ColorSpace/DMXPixelMappingColorSpace_RGBCMY.cpp
+++ b/Engine/Plugins/VirtualProduction/DMX/DMXPixelMapping/Source/DMXPixelMappingRuntime/Private/ColorSpace/DMXPixelMappingColorSpace_RGBCMY.cpp
@@ -2,6 +2,8 @@
 
 #include "ColorSpace/DMXPixelMappingColorSpace_RGBCMY.h"
 
+#include "ColorManagement/TransferFunctions.h"
+#include "DMXPixelMappingMainStreamObjectVersion.h"
 
 UDMXPixelMappingColorSpace_RGBCMY::UDMXPixelMappingColorSpace_RGBCMY()
 	: RedAttribute("Red")
@@ -12,14 +14,28 @@ UDMXPixelMappingColorSpace_RGBCMY::UDMXPixelMappingColorSpace_RGBCMY()
 
 void UDMXPixelMappingColorSpace_RGBCMY::SetRGBA(const FLinearColor& InColor)
 {
+	using namespace UE::Color;
+
 	if (bUseWorkingColorSpaceForInput && !InputColorSpace.Equals(UE::Color::FColorSpace::GetWorking()))
 	{
 		// Update in case the working color space changed
 		UpdateColorSpaceAndTransform();
 	}
 
-	const FLinearColor CalibratedColor = ColorSpaceTransform->Apply(InColor);
+	FLinearColor CalibratedColor = ColorSpaceTransform->Apply(InColor);
+
+	// Apply gamma
+	if (OutputGamma == EDMXPixelMappingGamma_RGBCMY::AsOutputColorSpace)
+	{
+		const EColorSpace SelectedOutputColorSpace = ConvertToOutputColorSpaceEnum(PixelMappingOutputColorSpace);
+		EncodeGammaCorrection(PixelMappingOutputColorSpace, CalibratedColor);
+	}
+	else if (OutputGamma == EDMXPixelMappingGamma_RGBCMY::Custom)
+	{
+		EncodeGammaCorrection(CustomGamma, CalibratedColor);
+	}
 
+	// Buffer DMX values
 	if (RedAttribute.IsValid())
 	{
 		const float Value = bSendCyan ? FMath::Abs(CalibratedColor.R - 1.f) : CalibratedColor.R;
@@ -56,6 +72,24 @@ void UDMXPixelMappingColorSpace_RGBCMY::SetRGBA(const FLinearColor& InColor)
 	}
 }
 
+void UDMXPixelMappingColorSpace_RGBCMY::Serialize(FArchive& Ar)
+{
+	Super::Serialize(Ar);
+
+	Ar.UsingCustomVersion(FDMXPixelMappingMainStreamObjectVersion::GUID);
+
+#if WITH_EDITOR
+	if (Ar.IsLoading())
+	{
+		// Assets created before 5.5 used linear gamma
+		if (Ar.CustomVer(FDMXPixelMappingMainStreamObjectVersion::GUID) < FDMXPixelMappingMainStreamObjectVersion::DMXOutputAppliesGammaCorrection)
+		{
+			OutputGamma = EDMXPixelMappingGamma_RGBCMY::Linear;
+		}
+	}
+#endif
+}
+
 void UDMXPixelMappingColorSpace_RGBCMY::PostLoad()
 {
 	Super::PostLoad();
@@ -138,3 +172,65 @@ UE::Color::EColorSpace UDMXPixelMappingColorSpace_RGBCMY::ConvertToOutputColorSp
 
 	return EColorSpace::None;
 }
+
+void UDMXPixelMappingColorSpace_RGBCMY::EncodeGammaCorrection(EDMXPixelMappingOutputColorSpace_RGBCMY InPixelMappingOutputColorSpaceEnum, FLinearColor& InOutColor)
+{
+	using namespace UE::Color;
+
+	switch (InPixelMappingOutputColorSpaceEnum)
+	{
+	case EDMXPixelMappingOutputColorSpace_RGBCMY::P3D65:
+		// Same as sRGB, fall through
+
+	case EDMXPixelMappingOutputColorSpace_RGBCMY::sRGB:
+		InOutColor.R = EncodeSRGB(InOutColor.R);
+		InOutColor.G = EncodeSRGB(InOutColor.G);
+		InOutColor.B = EncodeSRGB(InOutColor.B);
+		break;
+
+	case EDMXPixelMappingOutputColorSpace_RGBCMY::Rec2020:
+		InOutColor.R = EncodeRec2020(InOutColor.R);
+		InOutColor.G = EncodeRec2020(InOutColor.G);
+		InOutColor.B = EncodeRec2020(InOutColor.B);
+		break;
+
+	case EDMXPixelMappingOutputColorSpace_RGBCMY::P3DCI:
+		InOutColor.R = EncodeGamma26(InOutColor.R);
+		InOutColor.G = EncodeGamma26(InOutColor.G);
+		InOutColor.B = EncodeGamma26(InOutColor.B);
+		break;
+
+	case EDMXPixelMappingOutputColorSpace_RGBCMY::Plasa:
+		// PLASA ANSI E1.54 is linear
+		break;
+
+	default:
+		ensureMsgf(0, TEXT("Missing transfer function for color space. Cannot apply gamma correction."));
+	}
+}
+
+void UDMXPixelMappingColorSpace_RGBCMY::EncodeGammaCorrection(float InCustomGamma, FLinearColor& InOutColor)
+{
+	InOutColor.R = FMath::Pow(InOutColor.R, 1.f / InCustomGamma);
+	InOutColor.G = FMath::Pow(InOutColor.G, 1.f / InCustomGamma);
+	InOutColor.B = FMath::Pow(InOutColor.B, 1.f / InCustomGamma);
+}
+
+float UDMXPixelMappingColorSpace_RGBCMY::EncodeRec2020(float InValue) const
+{
+	// Values are expected to be in the 0-1 range
+	const float ClampedValue = FMath::Clamp(InValue, 0.f, 1.f);
+
+	constexpr float Beta = 0.018053968510807;
+	constexpr float Alpha = 1 + 5.5 * Beta;
+
+	if (ClampedValue < Beta)
+	{
+		return 4.5f * ClampedValue;
+	}
+	else
+	{
+		return Alpha * FMath::Pow(ClampedValue, 0.45f) - (Alpha - 1.f);
+	}
+}
+
diff --git a/Engine/Plugins/VirtualProduction/DMX/DMXPixelMapping/Source/DMXPixelMappingRuntime/Private/ColorSpace/DMXPixelMappingColorSpace_XYZ.cpp b/Engine/Plugins/VirtualProduction/DMX/DMXPixelMapping/Source/DMXPixelMappingRuntime/Private/ColorSpace/DMXPixelMappingColorSpace_XYZ.cpp
--- a/Engine/Plugins/VirtualProduction/DMX/DMXPixelMapping/Source/DMXPixelMappingRuntime/Private/ColorSpace/DMXPixelMappingColorSpace_XYZ.cpp
+++ b/Engine/Plugins/VirtualProduction/DMX/DMXPixelMapping/Source/DMXPixelMappingRuntime/Private/ColorSpace/DMXPixelMappingColorSpace_XYZ.cpp
@@ -17,8 +17,18 @@ void UDMXPixelMappingColorSpace_XYZ::SetRGBA(const FLinearColor& InColor)
 
 	// Convert RGB to CIEXYZ
 	const FMatrix44d Matrix = InputColorSpace.GetRgbToXYZ();
-	const FVector4 XYZW = Matrix.TransformVector(FVector(InColor));
+	FVector4 XYZW = Matrix.TransformVector(FVector(InColor));
+	
+	// Apply gamma to Y (the Z-Component of the vector)
+	if (!FMath::IsNearlyEqual(CustomGamma, 1.f))
+	{
+		FVector3d xyY = UE::Color::XYZToxyY(XYZW);
+		xyY.Z = FMath::Pow(xyY.Z, 1.f / CustomGamma);
+		
+		XYZW = UE::Color::xyYToXYZ(xyY);
+	}
 
+	// Buffer DMX values
 	if (XAttribute.IsValid())
 	{
 		SetAttributeValue(XAttribute, XYZW.X);
diff --git a/Engine/Plugins/VirtualProduction/DMX/DMXPixelMapping/Source/DMXPixelMappingRuntime/Private/ColorSpace/DMXPixelMappingColorSpace_xyY.cpp b/Engine/Plugins/VirtualProduction/DMX/DMXPixelMapping/Source/DMXPixelMappingRuntime/Private/ColorSpace/DMXPixelMappingColorSpace_xyY.cpp
--- a/Engine/Plugins/VirtualProduction/DMX/DMXPixelMapping/Source/DMXPixelMappingRuntime/Private/ColorSpace/DMXPixelMappingColorSpace_xyY.cpp
+++ b/Engine/Plugins/VirtualProduction/DMX/DMXPixelMapping/Source/DMXPixelMappingRuntime/Private/ColorSpace/DMXPixelMappingColorSpace_xyY.cpp
@@ -19,8 +19,15 @@ void UDMXPixelMappingColorSpace_xyY::SetRGBA(const FLinearColor& InColor)
 	// Convert RGB to CIE XYZ
 	const FMatrix44d& Matrix = InputColorSpace.GetRgbToXYZ();
 	const FVector4 XYZW = Matrix.TransformVector(FVector(InColor));
-	const FVector3d xyY = UE::Color::XYZToxyY(XYZW);
+	FVector3d xyY = UE::Color::XYZToxyY(XYZW);
 
+	// Apply gamma to Y (the Z-Component of the vector)
+	if (!FMath::IsNearlyEqual(CustomGamma, 1.f))
+	{
+		xyY.Z = FMath::Pow(xyY.Z, 1.f / CustomGamma);
+	}
+
+	// Buffer DMX values
 	if (!ensureMsgf(ColorSpaceRange != 0.0, TEXT("Coversion in PixelMapping Color Space xyY failed. Color space range is 0.")))
 	{
 		return;
diff --git a/Engine/Plugins/VirtualProduction/DMX/DMXPixelMapping/Source/DMXPixelMappingRuntime/Private/DMXPixelMappingMainStreamObjectVersion.h b/Engine/Plugins/VirtualProduction/DMX/DMXPixelMapping/Source/DMXPixelMappingRuntime/Private/DMXPixelMappingMainStreamObjectVersion.h
--- a/Engine/Plugins/VirtualProduction/DMX/DMXPixelMapping/Source/DMXPixelMappingRuntime/Private/DMXPixelMappingMainStreamObjectVersion.h
+++ b/Engine/Plugins/VirtualProduction/DMX/DMXPixelMapping/Source/DMXPixelMappingRuntime/Private/DMXPixelMappingMainStreamObjectVersion.h
@@ -28,6 +28,9 @@ struct FDMXPixelMappingMainStreamObjectVersion
 		// The renderer component now holds the layout rect, so its children can be layout out relative to the current texture size
 		RendererComponentHoldsLayoutRect,
 
+		// 5.5: Color spaces now apply a gamma correction
+		DMXOutputAppliesGammaCorrection,
+
 		// -----<new versions can be added above this line>-------------------------------------------------
 		VersionPlusOne,
 		LatestVersion = VersionPlusOne - 1
