DMX: Change the Pixel Mapping default texture to a UE5 logo

diff --git a/Engine/Plugins/VirtualProduction/DMX/DMXPixelMapping/Source/DMXPixelMappingRuntime/Private/Components/DMXPixelMappingRendererComponent.cpp b/Engine/Plugins/VirtualProduction/DMX/DMXPixelMapping/Source/DMXPixelMappingRuntime/Private/Components/DMXPixelMappingRendererComponent.cpp
--- a/Engine/Plugins/VirtualProduction/DMX/DMXPixelMapping/Source/DMXPixelMappingRuntime/Private/Components/DMXPixelMappingRendererComponent.cpp
+++ b/Engine/Plugins/VirtualProduction/DMX/DMXPixelMapping/Source/DMXPixelMappingRuntime/Private/Components/DMXPixelMappingRendererComponent.cpp
@@ -40,8 +40,9 @@ DECLARE_CYCLE_STAT(TEXT("PixelMapping RenderInputTexture"), STAT_DMXPixelMapping
 
 UDMXPixelMappingRendererComponent::UDMXPixelMappingRendererComponent()
 {
-	ConstructorHelpers::FObjectFinder<UTexture> DefaultTexture(TEXT("Texture2D'/Engine/VREditor/Devices/Vive/UE4_Logo.UE4_Logo'"), LOAD_NoWarn);
-	if (ensureAlwaysMsgf(DefaultTexture.Succeeded(), TEXT("Failed to load Texture2D'/Engine/VREditor/Devices/Vive/UE4_Logo.UE4_Logo'")))
+	static constexpr const TCHAR* UnrealEngineLogoPath = TEXT("Texture2D'/DMXPixelMapping/Textures/T_UnrealEngineLogo.T_UnrealEngineLogo'");
+	ConstructorHelpers::FObjectFinder<UTexture> DefaultTexture(UnrealEngineLogoPath, LOAD_NoWarn);
+	if (ensureAlwaysMsgf(DefaultTexture.Succeeded(), TEXT("Failed to load Default Texture for Pixel Mapping from '%s'"), UnrealEngineLogoPath))
 	{
 		InputTexture = DefaultTexture.Object;
 		RendererType = EDMXPixelMappingRendererType::Texture;
