DMX - Pixel mapping assets can now be dropped onto the level to spawn a Pixel Mapping Actor

* Pixel Mapping Actors offer auto-activate and play in editor options
* Pixel Mapping Actors offer blueprint nodes to control them from blueprints: 'Start Sending DMX', "Pause Sending DMX', 'Stop Sending DMX', 'Set Stop Mode'

diff --git a/Engine/Plugins/VirtualProduction/DMX/DMXPixelMapping/Source/DMXPixelMappingEditor/Private/Factories/DMXPixelMappingActorFactory.cpp b/Engine/Plugins/VirtualProduction/DMX/DMXPixelMapping/Source/DMXPixelMappingEditor/Private/Factories/DMXPixelMappingActorFactory.cpp
new file mode 100644
--- /dev/null
+++ b/Engine/Plugins/VirtualProduction/DMX/DMXPixelMapping/Source/DMXPixelMappingEditor/Private/Factories/DMXPixelMappingActorFactory.cpp
@@ -0,0 +1,78 @@
+// Copyright Epic Games, Inc. All Rights Reserved.
+
+#include "Factories/DMXPixelMappingActorFactory.h"
+
+#include "DMXPixelMapping.h"
+#include "DMXPixelMappingActor.h"
+#include "Components/StaticMeshComponent.h"
+#include "Engine/Level.h"
+#include "Engine/StaticMeshActor.h"
+#include "Engine/StaticMesh.h"
+
+#define LOCTEXT_NAMESPACE "DMXPixelMappingActorFactory"
+
+UDMXPixelMappingActorFactory::UDMXPixelMappingActorFactory()
+{
+	DisplayName = LOCTEXT("DMXPixelMappingActorFactoryDisplayName", "Pixel Mapping Actor");
+	NewActorClass = ADMXPixelMappingActor::StaticClass();
+}
+
+bool UDMXPixelMappingActorFactory::CanCreateActorFrom(const FAssetData& AssetData, FText& OutErrorMsg)
+{
+	if (AssetData.IsValid())
+	{
+		const UClass* AssetClass = AssetData.GetClass();
+		if (AssetClass && AssetClass == UDMXPixelMapping::StaticClass())
+		{
+			return true;
+		}
+	}
+
+	return false;
+}
+
+AActor* UDMXPixelMappingActorFactory::SpawnActor(UObject* Asset, ULevel* InLevel, const FTransform& Transform, const FActorSpawnParameters& InSpawnParams)
+{
+	UDMXPixelMapping* PixelMapping = Cast<UDMXPixelMapping>(Asset);
+	if (!PixelMapping)
+	{
+		return nullptr;
+	}
+
+	if (InSpawnParams.ObjectFlags & RF_Transient)
+	{
+		// This is a hack for drag and drop so that we don't spawn all the actors for the preview actor since it gets deleted right after.
+		FActorSpawnParameters SpawnInfo(InSpawnParams);
+		SpawnInfo.OverrideLevel = InLevel;
+		AStaticMeshActor* DragActor = Cast<AStaticMeshActor>(InLevel->GetWorld()->SpawnActor(AStaticMeshActor::StaticClass(), &Transform, SpawnInfo));
+		DragActor->GetStaticMeshComponent()->SetStaticMesh(Cast< UStaticMesh >(FSoftObjectPath(TEXT("StaticMesh'/Engine/EditorMeshes/EditorSphere.EditorSphere'")).TryLoad()));
+		DragActor->SetActorScale3D(FVector(0.1f));
+
+		return DragActor;
+	}
+	else 
+	{
+		if (!InLevel)
+		{
+			return nullptr;
+		}
+
+		UWorld* World = InLevel->GetWorld();
+		if (!World)
+		{
+			return nullptr;
+		}
+
+		ADMXPixelMappingActor* PixelMappingActor = World->SpawnActor<ADMXPixelMappingActor>(ADMXPixelMappingActor::StaticClass());
+		if (PixelMappingActor)
+		{
+			PixelMappingActor->SetPixelMapping(PixelMapping);
+
+			return PixelMappingActor;
+		}
+	}
+
+	return nullptr;
+}
+
+#undef LOCTEXT_NAMESPACE
diff --git a/Engine/Plugins/VirtualProduction/DMX/DMXPixelMapping/Source/DMXPixelMappingEditor/Private/Factories/DMXPixelMappingActorFactory.h b/Engine/Plugins/VirtualProduction/DMX/DMXPixelMapping/Source/DMXPixelMappingEditor/Private/Factories/DMXPixelMappingActorFactory.h
new file mode 100644
--- /dev/null
+++ b/Engine/Plugins/VirtualProduction/DMX/DMXPixelMapping/Source/DMXPixelMappingEditor/Private/Factories/DMXPixelMappingActorFactory.h
@@ -0,0 +1,29 @@
+// Copyright Epic Games, Inc. All Rights Reserved.
+
+#pragma once
+
+#include "ActorFactories/ActorFactory.h"
+#include "UObject/ObjectMacros.h"
+
+#include "DMXPixelMappingActorFactory.generated.h"
+
+class AActor;
+struct FAssetData;
+
+
+/** Actor Factory for DMX Pixel Mapping Actor */
+UCLASS()
+class UDMXPixelMappingActorFactory
+	: public UActorFactory
+{
+	GENERATED_BODY()
+
+public:
+	/** Constructor */
+	UDMXPixelMappingActorFactory();
+
+	//~ Begin UActorFactory Interface
+	virtual bool CanCreateActorFrom(const FAssetData& AssetData, FText& OutErrorMsg) override;
+	virtual AActor* SpawnActor(UObject* Asset, ULevel* InLevel, const FTransform& Transform, const FActorSpawnParameters& InSpawnParams) override;
+	//~ End UActorFactory Interface
+};
diff --git a/Engine/Plugins/VirtualProduction/DMX/DMXPixelMapping/Source/DMXPixelMappingEditor/Private/Settings/DMXPixelMappingEditorSettings.h b/Engine/Plugins/VirtualProduction/DMX/DMXPixelMapping/Source/DMXPixelMappingEditor/Private/Settings/DMXPixelMappingEditorSettings.h
--- a/Engine/Plugins/VirtualProduction/DMX/DMXPixelMapping/Source/DMXPixelMappingEditor/Private/Settings/DMXPixelMappingEditorSettings.h
+++ b/Engine/Plugins/VirtualProduction/DMX/DMXPixelMapping/Source/DMXPixelMappingEditor/Private/Settings/DMXPixelMappingEditorSettings.h
@@ -90,6 +90,7 @@ public:
 	FDMXPixelMappingDesignerSettings DesignerSettings;
 
 	/** Reset DMX mode to be used in editor */
+	UE_DEPRECATED(5.6, "Each pixel mapping holds its own stop mode now. See UDMXPixelMapping::SetResetDMXMode")
 	UPROPERTY(Config)
 	EDMXPixelMappingResetDMXMode EditorResetDMXMode;
 
diff --git a/Engine/Plugins/VirtualProduction/DMX/DMXPixelMapping/Source/DMXPixelMappingEditor/Private/Toolkits/DMXPixelMappingToolkit.cpp b/Engine/Plugins/VirtualProduction/DMX/DMXPixelMapping/Source/DMXPixelMappingEditor/Private/Toolkits/DMXPixelMappingToolkit.cpp
--- a/Engine/Plugins/VirtualProduction/DMX/DMXPixelMapping/Source/DMXPixelMappingEditor/Private/Toolkits/DMXPixelMappingToolkit.cpp
+++ b/Engine/Plugins/VirtualProduction/DMX/DMXPixelMapping/Source/DMXPixelMappingEditor/Private/Toolkits/DMXPixelMappingToolkit.cpp
@@ -17,6 +17,7 @@
 #include "DMXPixelMappingEditorModule.h"
 #include "DMXPixelMappingEditorStyle.h"
 #include "DMXPixelMappingEditorUtils.h"
+#include "DMXPixelMappingMainStreamObjectVersion.h"
 #include "DMXPixelMappingToolbar.h"
 #include "Engine/Texture.h"
 #include "Engine/Texture2D.h"
@@ -72,6 +73,15 @@ FDMXPixelMappingToolkit::~FDMXPixelMappingToolkit()
 void FDMXPixelMappingToolkit::InitPixelMappingEditor(const EToolkitMode::Type Mode, const TSharedPtr<class IToolkitHost>& InitToolkitHost, UDMXPixelMapping* InDMXPixelMapping)
 {
 	check(InDMXPixelMapping);
+
+	// Upgrade to use a per pixel mapping DMX reset mode
+	if (InDMXPixelMapping->GetLinkerCustomVersion(FDMXPixelMappingMainStreamObjectVersion::GUID) < FDMXPixelMappingMainStreamObjectVersion::PerPixelMappingResetDMXMode)
+	{
+		PRAGMA_DISABLE_DEPRECATION_WARNINGS
+		InDMXPixelMapping->SetResetDMXMode(GetDefault<UDMXPixelMappingEditorSettings>()->EditorResetDMXMode);
+		PRAGMA_ENABLE_DEPRECATION_WARNINGS
+	}
+
 	InDMXPixelMapping->DestroyInvalidComponents();
 
 	// Make sure we loaded all UObjects
@@ -275,10 +285,6 @@ void FDMXPixelMappingToolkit::Tick(float DeltaTime)
 
 	// Render, send DMX if required
 	RootComponent->Render();
-	if (bIsPlayingDMX && !bIsPaused)
-	{
-		RootComponent->SendDMX();
-	}
 
 	// Detect and broadcast editor setting changes
 	if (FMemory::Memcmp(EditorSettingsDump.GetData(), GetDefault<UDMXPixelMappingEditorSettings>(), sizeof(UDMXPixelMappingEditorSettings)) != 0)
@@ -424,66 +430,63 @@ void FDMXPixelMappingToolkit::AddRenderer()
 
 void FDMXPixelMappingToolkit::PlayDMX()
 {
-	bIsPaused = false;
-	bIsPlayingDMX = true;
+	UDMXPixelMapping* PixelMapping = GetDMXPixelMapping();
+	if (PixelMapping)
+	{
+		PixelMapping->StartSendingDMX();
+	}
 }
 
 void FDMXPixelMappingToolkit::PauseDMX()
 {
-	bIsPaused = true;
-	bIsPlayingDMX = false;
+	UDMXPixelMapping* PixelMapping = GetDMXPixelMapping();
+	if (PixelMapping)
+	{
+		PixelMapping->PauseSendingDMX();
+	}
 }
 
 void FDMXPixelMappingToolkit::StopPlayingDMX()
 {
 	UDMXPixelMapping* PixelMapping = GetDMXPixelMapping();
-	if (!PixelMapping)
+	if (PixelMapping)
 	{
-		return;
+		PixelMapping->StopSendingDMX();
 	}
-
-	UDMXPixelMappingRootComponent* RootComponent = PixelMapping->RootComponent;
-	if (!ensure(RootComponent))
-	{
-		return;
-	}
-
-	const UDMXPixelMappingEditorSettings* Settings = GetDefault<UDMXPixelMappingEditorSettings>();
-	RootComponent->ResetDMX(Settings->EditorResetDMXMode);
-
-	bIsPaused = false;
-	bIsPlayingDMX = false;
 }
 
 void FDMXPixelMappingToolkit::TogglePlayPauseDMX()
 {
-	if (bIsPlayingDMX)
+	UDMXPixelMapping* PixelMapping = GetDMXPixelMapping();
+	if (PixelMapping && PixelMapping->IsPlayingDMX())
 	{
-		PauseDMX();
+		PixelMapping->PauseSendingDMX();
 	}
-	else
+	else if (PixelMapping)
 	{
-		PlayDMX();
+		PixelMapping->StartSendingDMX();
 	}
 }
 
 void FDMXPixelMappingToolkit::TogglePlayStopDMX()
 {
-	if (bIsPlayingDMX)
+	UDMXPixelMapping* PixelMapping = GetDMXPixelMapping();
+	if (PixelMapping && PixelMapping->IsPlayingDMX())
 	{
-		StopPlayingDMX();
+		PixelMapping->StopSendingDMX();
 	}
-	else
+	else if (PixelMapping)
 	{
-		PlayDMX();
+		PixelMapping->StartSendingDMX();
 	}
 }
 
 void FDMXPixelMappingToolkit::SetEditorResetDMXMode(EDMXPixelMappingResetDMXMode NewMode)
 {
-	UDMXPixelMappingEditorSettings* EditorSettings = GetMutableDefault<UDMXPixelMappingEditorSettings>();
-	EditorSettings->EditorResetDMXMode = NewMode;
-	EditorSettings->SaveConfig();
+	if (UDMXPixelMapping* PixelMapping = GetDMXPixelMapping())
+	{
+		PixelMapping->SetResetDMXMode(NewMode);
+	}
 }
 
 void FDMXPixelMappingToolkit::UpdateBlueprintNodes() const
@@ -1017,6 +1020,13 @@ TSharedRef<SDMXPixelMappingLayoutView> FDMXPixelMappingToolkit::GetOrCreateLayou
 	return LayoutView.ToSharedRef();
 }
 
+bool FDMXPixelMappingToolkit::IsPlayingDMX() const
+{
+	const UDMXPixelMapping* PixelMapping = GetDMXPixelMapping();
+
+	return PixelMapping && PixelMapping->IsPlayingDMX();
+}
+
 #define UE_DMX_MAP_EDITOR_SETTING_TO_TOGGLE_COMMAND(EditorSettings, StructName, MemberName, Action) \
 	GetToolkitCommands()->MapAction( \
 		FDMXPixelMappingEditorCommands::Get().Action, \
@@ -1053,11 +1063,19 @@ void FDMXPixelMappingToolkit::SetupCommands()
 		FExecuteAction::CreateSP(this, &FDMXPixelMappingToolkit::PlayDMX),
 		FCanExecuteAction::CreateLambda([this]
 			{
+				const UDMXPixelMapping* PixelMapping = GetDMXPixelMapping();
+				const bool bIsPlayingDMX = PixelMapping && PixelMapping->IsPlayingDMX();
+				const bool bIsPaused = PixelMapping && PixelMapping->IsPaused();
+
 				return !bIsPlayingDMX && !bIsPaused;
 			}),
 		FIsActionChecked(),
 		FIsActionButtonVisible::CreateLambda([this]
-			{
+			{				
+				const UDMXPixelMapping* PixelMapping = GetDMXPixelMapping();
+				const bool bIsPlayingDMX = PixelMapping && PixelMapping->IsPlayingDMX();
+				const bool bIsPaused = PixelMapping && PixelMapping->IsPaused();
+
 				return !bIsPlayingDMX && !bIsPaused;
 			})
 	);
@@ -1067,11 +1085,17 @@ void FDMXPixelMappingToolkit::SetupCommands()
 		FExecuteAction::CreateSP(this, &FDMXPixelMappingToolkit::PauseDMX),
 		FCanExecuteAction::CreateLambda([this] 
 			{ 
+				const UDMXPixelMapping* PixelMapping = GetDMXPixelMapping();
+				const bool bIsPlayingDMX = PixelMapping && PixelMapping->IsPlayingDMX();
+
 				return bIsPlayingDMX; 
 			}),
 		FIsActionChecked(),
 		FIsActionButtonVisible::CreateLambda([this]
 			{
+				const UDMXPixelMapping* PixelMapping = GetDMXPixelMapping();
+				const bool bIsPlayingDMX = PixelMapping && PixelMapping->IsPlayingDMX();
+
 				return bIsPlayingDMX;
 			})
 	);
@@ -1081,11 +1105,19 @@ void FDMXPixelMappingToolkit::SetupCommands()
 		FExecuteAction::CreateSP(this, &FDMXPixelMappingToolkit::PlayDMX),
 		FCanExecuteAction::CreateLambda([this]
 			{
+				const UDMXPixelMapping* PixelMapping = GetDMXPixelMapping();
+				const bool bIsPlayingDMX = PixelMapping && PixelMapping->IsPlayingDMX();
+				const bool bIsPaused = PixelMapping && PixelMapping->IsPaused();
+
 				return !bIsPlayingDMX && bIsPaused;
 			}),
 		FIsActionChecked(),
 		FIsActionButtonVisible::CreateLambda([this]
 			{
+				const UDMXPixelMapping* PixelMapping = GetDMXPixelMapping();
+				const bool bIsPlayingDMX = PixelMapping && PixelMapping->IsPlayingDMX();
+				const bool bIsPaused = PixelMapping && PixelMapping->IsPaused();
+
 				return !bIsPlayingDMX && bIsPaused;
 			})
 	);
@@ -1095,6 +1127,10 @@ void FDMXPixelMappingToolkit::SetupCommands()
 		FExecuteAction::CreateSP(this, &FDMXPixelMappingToolkit::StopPlayingDMX),
 		FCanExecuteAction::CreateLambda([this]
 			{
+				const UDMXPixelMapping* PixelMapping = GetDMXPixelMapping();
+				const bool bIsPlayingDMX = PixelMapping && PixelMapping->IsPlayingDMX();
+				const bool bIsPaused = PixelMapping && PixelMapping->IsPaused();
+
 				return bIsPlayingDMX || bIsPaused;
 			})
 	);
@@ -1200,8 +1236,12 @@ void FDMXPixelMappingToolkit::ExtendToolbar()
 
 ECheckBoxState FDMXPixelMappingToolkit::GetEditorResetDMXModeCheckboxState(EDMXPixelMappingResetDMXMode CompareMode) const
 {
-	const UDMXPixelMappingEditorSettings* EditorSettings = GetDefault<UDMXPixelMappingEditorSettings>();
-	return EditorSettings->EditorResetDMXMode == CompareMode ? ECheckBoxState::Checked : ECheckBoxState::Unchecked;
+	if (UDMXPixelMapping* PixelMapping = GetDMXPixelMapping())
+	{
+		return PixelMapping->GetResetDMXMode() == CompareMode ? ECheckBoxState::Checked : ECheckBoxState::Unchecked;
+	}
+
+	return ECheckBoxState::Undetermined;
 }
 
 ECheckBoxState FDMXPixelMappingToolkit::GetTransformHandleModeCheckboxState(EDMXPixelMappingTransformHandleMode CompareTransformHandleMode) const
diff --git a/Engine/Plugins/VirtualProduction/DMX/DMXPixelMapping/Source/DMXPixelMappingEditor/Private/Toolkits/DMXPixelMappingToolkit.h b/Engine/Plugins/VirtualProduction/DMX/DMXPixelMapping/Source/DMXPixelMappingEditor/Private/Toolkits/DMXPixelMappingToolkit.h
--- a/Engine/Plugins/VirtualProduction/DMX/DMXPixelMapping/Source/DMXPixelMappingEditor/Private/Toolkits/DMXPixelMappingToolkit.h
+++ b/Engine/Plugins/VirtualProduction/DMX/DMXPixelMapping/Source/DMXPixelMappingEditor/Private/Toolkits/DMXPixelMappingToolkit.h
@@ -132,7 +132,8 @@ public:
 	/** Gets or creates the Layout View for this Pixel Mapping instance */
 	TSharedRef<SDMXPixelMappingLayoutView> GetOrCreateLayoutView();
 
-	bool IsPlayingDMX() const { return bIsPlayingDMX; }
+	/** Returns true if the pixel mapping asset is currently playing DMX */
+	bool IsPlayingDMX() const;
 
 	void SetActiveRenderComponent(UDMXPixelMappingRendererComponent* InComponent);
 
@@ -292,12 +293,6 @@ private:
 	/** Command list for handling widget actions in the PixelMapping Toolkit */
 	TSharedPtr<FUICommandList> DesignerCommandList;
 
-	/** True while playing DMX */
-	bool bIsPlayingDMX = false;
-
-	/** True if paused */
-	bool bIsPaused = false;
-
 	/** True while adding components (to avoid needlessly updating blueprint nodes on each component added via our own methods) */
 	bool bAddingComponents = false;
 
diff --git a/Engine/Plugins/VirtualProduction/DMX/DMXPixelMapping/Source/DMXPixelMappingRuntime/Internal/DMXPixelMappingActor.h b/Engine/Plugins/VirtualProduction/DMX/DMXPixelMapping/Source/DMXPixelMappingRuntime/Internal/DMXPixelMappingActor.h
new file mode 100644
--- /dev/null
+++ b/Engine/Plugins/VirtualProduction/DMX/DMXPixelMapping/Source/DMXPixelMappingRuntime/Internal/DMXPixelMappingActor.h
@@ -0,0 +1,85 @@
+// Copyright Epic Games, Inc. All Rights Reserved.
+
+#pragma once
+
+#include "GameFramework/Actor.h"
+
+#include "DMXPixelMappingActor.generated.h"
+
+class UDMXPixelMapping;
+class USceneComponent;
+enum class EDMXPixelMappingResetDMXMode : uint8;
+
+/** Actor class for DMX Pixel Mapping */
+UCLASS(NotBlueprintable)
+class DMXPIXELMAPPINGRUNTIME_API ADMXPixelMappingActor
+	: public AActor
+{
+	GENERATED_BODY()
+
+public:
+	/** Constructor */
+	ADMXPixelMappingActor();
+
+	/** Sets the Pixel Mapping used in this actor */
+	void SetPixelMapping(UDMXPixelMapping* InPixelMapping);
+
+	/** Starts sending DMX */
+	UFUNCTION(BlueprintCallable, Category = "DMX Pixel Mapping")
+	void StartSendingDMX();
+
+	/** Stop sending DMX. */
+	UFUNCTION(BlueprintCallable, Category = "DMX Pixel Mapping")
+	void StopSendingDMX();
+
+	/** Pause sending DMX */
+	UFUNCTION(BlueprintCallable, Category = "DMX Pixel Mapping")
+	void PauseSendingDMX();
+
+	/** Sets how the pixel mapping asset resets the channels it sends to when Stop Sending DMX is called */
+	UFUNCTION(BlueprintCallable, Category = "DMX Pixel Mapping")
+	void SetStopMode(EDMXPixelMappingResetDMXMode ResetMode);
+
+protected:
+	//~ Begin AActor interface
+	virtual void PostLoad() override;
+	virtual void BeginPlay() override;
+	virtual void EndPlay(const EEndPlayReason::Type EndPlayReason) override;
+#if WITH_EDITOR
+	virtual void PostEditChangeProperty(FPropertyChangedEvent& PropertyChangedEvent) override;
+#endif
+	//~ End AActor interface
+
+private:
+#if WITH_EDITOR
+	/** Applies the play in editor state. Enables send when auto activate is on, send DMX in editor is enabled and not currently playing */
+	void ApplySendDMXInEditorState();
+#endif // WITH_EDITOR
+
+	/** The pixel mapping used in this actor */
+	UPROPERTY(VisibleAnywhere, Category = "DMX Pixel Mapping")
+	TObjectPtr<UDMXPixelMapping> PixelMapping;
+
+	/**  
+	 * True if the pixel mapping actor auto activates.
+	 * If Send DMX in Editor is disabled, starts to send DMX on begin play.
+	 * If Send DMX in Editor is enabled, it will start to send DMX when the level is loaded in editor.
+	 */
+	UPROPERTY(EditAnywhere, Category = "DMX Pixel Mapping")
+	bool bAutoActivate = true;
+
+#if WITH_EDITORONLY_DATA
+	/** True if the pixel mapping should send DMX data in editor */
+	UPROPERTY(EditAnywhere, Category = "DMX Pixel Mapping", Meta = (DisplayName = "Send DMX in Editor"))
+	bool bSendDMXInEditor = false;
+#endif // WITH_EDITORONLY_DATA
+
+#if WITH_EDITOR
+	/** True while the actor plays in a world */
+	bool bIsPlayInWorld = false;
+#endif // WITH_EDITOR
+
+	/** Scene component to make the actor easily visible in Editor */
+	UPROPERTY(VisibleAnywhere, Category = "Actor", AdvancedDisplay, Meta = (AllowPrivateAccess = true))
+	TObjectPtr<USceneComponent> RootSceneComponent;
+};
diff --git a/Engine/Plugins/VirtualProduction/DMX/DMXPixelMapping/Source/DMXPixelMappingRuntime/Internal/DMXPixelMappingMainStreamObjectVersion.h b/Engine/Plugins/VirtualProduction/DMX/DMXPixelMapping/Source/DMXPixelMappingRuntime/Internal/DMXPixelMappingMainStreamObjectVersion.h
new file mode 100644
--- /dev/null
+++ b/Engine/Plugins/VirtualProduction/DMX/DMXPixelMapping/Source/DMXPixelMappingRuntime/Internal/DMXPixelMappingMainStreamObjectVersion.h
@@ -0,0 +1,47 @@
+// Copyright Epic Games, Inc. All Rights Reserved.
+#pragma once
+
+#include "CoreTypes.h"
+#include "Misc/Guid.h"
+
+
+// Custom serialization version for changes to DMX Pixel Mapping Objects in the Main Stream
+struct DMXPIXELMAPPINGRUNTIME_API FDMXPixelMappingMainStreamObjectVersion
+{
+	enum Type
+	{
+		// Roughly corresponds to 4.26
+		BeforeCustomVersionWasAdded = 0,
+
+		// Update DMXPixelMappingMatrixComponent to no longer use a custom MatrixFixturePatchRef, but the default EntityFixturePatchRef one.
+		ChangePixelMappingMatrixComponentToFixturePatchReference,
+
+		// Update DMXPixelMappingBaseComponent to store the parent as a weak ptr
+		UseWeakPtrForPixelMappingComponentParent,
+
+		// Update DMXPixelMappingRendererComponent to lock those components that use a texture in designer to trigger the edit condition of the size property
+		LockRendererComponentsThatUseTextureInDesigner,
+
+		// Update DMXPixelMappingFixtureGroupItemComponent and DMXPixelMappingMatrixCellComponent to use DMXPixelMappingColorSpace
+		UseDMXPixelMappingColorSpace,
+
+		// The renderer component now holds the layout rect, so its children can be layout out relative to the current texture size
+		RendererComponentHoldsLayoutRect,
+
+		// 5.5: Color spaces now apply a gamma correction
+		DMXOutputAppliesGammaCorrection,
+
+		// 5.6: Each pixel mapping holds its own reset DMX mode
+		PerPixelMappingResetDMXMode,
+
+		// -----<new versions can be added above this line>-------------------------------------------------
+		VersionPlusOne,
+		LatestVersion = VersionPlusOne - 1
+	};
+
+	// The GUID for this custom version number
+	const static FGuid GUID;
+
+private:
+	FDMXPixelMappingMainStreamObjectVersion() {}
+};
diff --git a/Engine/Plugins/VirtualProduction/DMX/DMXPixelMapping/Source/DMXPixelMappingRuntime/Private/Components/DMXPixelMappingRootComponent.cpp b/Engine/Plugins/VirtualProduction/DMX/DMXPixelMapping/Source/DMXPixelMappingRuntime/Private/Components/DMXPixelMappingRootComponent.cpp
--- a/Engine/Plugins/VirtualProduction/DMX/DMXPixelMapping/Source/DMXPixelMappingRuntime/Private/Components/DMXPixelMappingRootComponent.cpp
+++ b/Engine/Plugins/VirtualProduction/DMX/DMXPixelMapping/Source/DMXPixelMappingRuntime/Private/Components/DMXPixelMappingRootComponent.cpp
@@ -71,13 +71,20 @@ void UDMXPixelMappingRootComponent::SendDMX()
 
 void UDMXPixelMappingRootComponent::Render()
 {
-	ForEachChild([&](UDMXPixelMappingBaseComponent* InComponent)
+	// Only ever render once per frame
+	if (LastRenderedFrameNumber < GFrameNumber ||
+		GFrameNumber == 0)
 	{
-		if (UDMXPixelMappingRendererComponent* Component = Cast<UDMXPixelMappingRendererComponent>(InComponent))
-		{
-			Component->Render();
-		}
-	}, false);
+		LastRenderedFrameNumber = GFrameNumber;
+	
+		ForEachChild([&](UDMXPixelMappingBaseComponent* InComponent)
+			{
+				if (UDMXPixelMappingRendererComponent* Component = Cast<UDMXPixelMappingRendererComponent>(InComponent))
+				{
+					Component->Render();
+				}
+			}, false);
+	}
 }
 
 void UDMXPixelMappingRootComponent::RenderAndSendDMX()
diff --git a/Engine/Plugins/VirtualProduction/DMX/DMXPixelMapping/Source/DMXPixelMappingRuntime/Private/DMXPixelMapping.cpp b/Engine/Plugins/VirtualProduction/DMX/DMXPixelMapping/Source/DMXPixelMappingRuntime/Private/DMXPixelMapping.cpp
--- a/Engine/Plugins/VirtualProduction/DMX/DMXPixelMapping/Source/DMXPixelMappingRuntime/Private/DMXPixelMapping.cpp
+++ b/Engine/Plugins/VirtualProduction/DMX/DMXPixelMapping/Source/DMXPixelMappingRuntime/Private/DMXPixelMapping.cpp
@@ -2,10 +2,12 @@
 
 #include "DMXPixelMapping.h"
 
+#include "Components/DMXPixelMappingRendererComponent.h"
 #include "Components/DMXPixelMappingRootComponent.h"
 #include "Components/DMXPixelMappingFixtureGroupItemComponent.h"
 #include "Components/DMXPixelMappingMatrixComponent.h"
 #include "Components/DMXPixelMappingMatrixCellComponent.h"
+#include "DMXPixelMappingMainStreamObjectVersion.h"
 #include "Library/DMXEntityFixturePatch.h"
 #include "UObject/LinkerLoad.h"
 
@@ -16,12 +18,43 @@
 
 
 UDMXPixelMapping::UDMXPixelMapping()
+	: ResetDMXMode(EDMXPixelMappingResetDMXMode::DoNotSendValues)
 {
 #if WITH_EDITOR
 	SnapGridColor = FLinearColor::White.CopyWithNewOpacity(.12f);
 #endif // WITH_EDITOR
 }
 
+void UDMXPixelMapping::StartSendingDMX()
+{
+	bIsPaused = false;
+	bIsPlayingDMX = true;
+}
+
+void UDMXPixelMapping::StopSendingDMX()
+{
+	if (RootComponent)
+	{
+		RootComponent->ResetDMX(ResetDMXMode);
+	}
+
+	bIsPaused = false;
+	bIsPlayingDMX = false;
+}
+
+void UDMXPixelMapping::PauseSendingDMX()
+{
+	bIsPaused = true;
+	bIsPlayingDMX = false;
+}
+
+void UDMXPixelMapping::Serialize(FArchive& Ar)
+{
+	Super::Serialize(Ar);
+
+	Ar.UsingCustomVersion(FDMXPixelMappingMainStreamObjectVersion::GUID);
+}
+
 void UDMXPixelMapping::PostLoad()
 {
 	Super::PostLoad();
@@ -151,6 +184,30 @@ void UDMXPixelMapping::RemoveComponent(UDMXPixelMappingBaseComponent* InComponen
 	}
 }
 
+ETickableTickType UDMXPixelMapping::GetTickableTickType() const
+{
+	return ETickableTickType::Conditional;
+}
+
+bool UDMXPixelMapping::IsTickable() const
+{
+	return bIsPlayingDMX && !bIsPaused;
+}
+
+void UDMXPixelMapping::Tick(float DeltaTime)
+{
+	TArray<UDMXPixelMappingRendererComponent*> RendererComponents;
+	GetAllComponentsOfClass<UDMXPixelMappingRendererComponent>(RendererComponents);
+
+	for (UDMXPixelMappingRendererComponent* RendererComponent : RendererComponents)
+	{
+		if (RendererComponent)
+		{
+			RendererComponent->RenderAndSendDMX();
+		}
+	}
+}
+
 void UDMXPixelMapping::ForEachComponent(TComponentPredicate Predicate) const
 {
 	if (RootComponent)
diff --git a/Engine/Plugins/VirtualProduction/DMX/DMXPixelMapping/Source/DMXPixelMappingRuntime/Private/DMXPixelMappingActor.cpp b/Engine/Plugins/VirtualProduction/DMX/DMXPixelMapping/Source/DMXPixelMappingRuntime/Private/DMXPixelMappingActor.cpp
new file mode 100644
--- /dev/null
+++ b/Engine/Plugins/VirtualProduction/DMX/DMXPixelMapping/Source/DMXPixelMappingRuntime/Private/DMXPixelMappingActor.cpp
@@ -0,0 +1,124 @@
+// Copyright Epic Games, Inc. All Rights Reserved.
+
+#include "DMXPixelMappingActor.h"
+
+#include "Components/DMXPixelMappingBaseComponent.h"
+#include "Components/SceneComponent.h"
+#include "DMXPixelMapping.h"
+
+ADMXPixelMappingActor::ADMXPixelMappingActor()
+{
+	RootSceneComponent = CreateDefaultSubobject<USceneComponent>("SceneComponent");
+	RootComponent = RootSceneComponent;
+}
+
+void ADMXPixelMappingActor::SetPixelMapping(UDMXPixelMapping* InPixelMapping)
+{
+	if (!ensureAlwaysMsgf(!PixelMapping, TEXT("Tried to set a Pixel Mapping for %s, but it already has one set. Changing the pixel mapping is not currently supported."), *GetName()))
+	{
+		return;
+	}
+
+	PixelMapping = InPixelMapping;
+}
+
+void ADMXPixelMappingActor::StartSendingDMX()
+{
+	if (PixelMapping)
+	{
+		PixelMapping->StartSendingDMX();
+	}
+}
+
+void ADMXPixelMappingActor::StopSendingDMX()
+{
+	if (PixelMapping)
+	{
+		PixelMapping->StopSendingDMX();
+	}
+}
+
+void ADMXPixelMappingActor::PauseSendingDMX()
+{
+	if (PixelMapping)
+	{
+		PixelMapping->PauseSendingDMX();
+	}
+}
+
+void ADMXPixelMappingActor::SetStopMode(EDMXPixelMappingResetDMXMode ResetMode)
+{
+	if (PixelMapping)
+	{
+		PixelMapping->SetResetDMXMode(ResetMode);
+	}
+}
+
+void ADMXPixelMappingActor::PostLoad()
+{
+	Super::PostLoad();
+
+#if WITH_EDITOR
+	ApplySendDMXInEditorState();
+#endif // WITH_EDITOR
+}
+
+void ADMXPixelMappingActor::BeginPlay()
+{
+	Super::BeginPlay();
+
+	if (bAutoActivate)
+	{
+		StartSendingDMX();
+	}
+
+#if WITH_EDITOR
+	bIsPlayInWorld = true;
+#endif 
+}
+
+void ADMXPixelMappingActor::EndPlay(const EEndPlayReason::Type EndPlayReason)
+{
+	Super::EndPlay(EndPlayReason);
+
+#if WITH_EDITOR
+	bIsPlayInWorld = false;
+	ApplySendDMXInEditorState();
+#else
+	StopSendingDMX();
+#endif
+}
+
+#if WITH_EDITOR
+void ADMXPixelMappingActor::PostEditChangeProperty(FPropertyChangedEvent& PropertyChangedEvent)
+{
+	Super::PostEditChangeProperty(PropertyChangedEvent);
+
+	const FName PropertyName = PropertyChangedEvent.GetPropertyName();
+	if (PropertyName == GET_MEMBER_NAME_CHECKED(ADMXPixelMappingActor, bSendDMXInEditor) ||
+		PropertyName == GET_MEMBER_NAME_CHECKED(ADMXPixelMappingActor, bAutoActivate))
+	{
+		ApplySendDMXInEditorState();
+	}
+}
+#endif // WITH_EDITOR
+
+#if WITH_EDITOR
+void ADMXPixelMappingActor::ApplySendDMXInEditorState()
+{
+	if (PixelMapping)
+	{
+		const bool bShouldSendDMX = bAutoActivate && (bIsPlayInWorld || bSendDMXInEditor);
+		if (bShouldSendDMX)
+		{
+			PixelMapping->StartSendingDMX();
+		}
+		else
+		{
+			PixelMapping->StopSendingDMX();
+		}
+	}
+}
+#endif // WITH_EDITOR
+
+#undef LOCTEXT_NAMESPACE
diff --git a/Engine/Plugins/VirtualProduction/DMX/DMXPixelMapping/Source/DMXPixelMappingRuntime/Private/DMXPixelMappingMainStreamObjectVersion.h b/Engine/Plugins/VirtualProduction/DMX/DMXPixelMapping/Source/DMXPixelMappingRuntime/Private/DMXPixelMappingMainStreamObjectVersion.h
deleted file mode 100644
diff --git a/Engine/Plugins/VirtualProduction/DMX/DMXPixelMapping/Source/DMXPixelMappingRuntime/Public/Components/DMXPixelMappingRootComponent.h b/Engine/Plugins/VirtualProduction/DMX/DMXPixelMapping/Source/DMXPixelMappingRuntime/Public/Components/DMXPixelMappingRootComponent.h
--- a/Engine/Plugins/VirtualProduction/DMX/DMXPixelMapping/Source/DMXPixelMappingRuntime/Public/Components/DMXPixelMappingRootComponent.h
+++ b/Engine/Plugins/VirtualProduction/DMX/DMXPixelMapping/Source/DMXPixelMappingRuntime/Public/Components/DMXPixelMappingRootComponent.h
@@ -49,4 +49,7 @@ private:
 	/** Chached Renderer Component names and their actual object pointer. Useful to speed up access (see DMXPixelMappingSubsystem) */
 	UPROPERTY(Transient)
 	TMap<FName, TObjectPtr<UDMXPixelMappingRendererComponent>> CachedRendererComponentsByName;
+
+	/** Holds the last rendered frame number, to avoid rendering multiple times in a single frame */
+	uint32 LastRenderedFrameNumber = 0;
 };
diff --git a/Engine/Plugins/VirtualProduction/DMX/DMXPixelMapping/Source/DMXPixelMappingRuntime/Public/DMXPixelMapping.h b/Engine/Plugins/VirtualProduction/DMX/DMXPixelMapping/Source/DMXPixelMappingRuntime/Public/DMXPixelMapping.h
--- a/Engine/Plugins/VirtualProduction/DMX/DMXPixelMapping/Source/DMXPixelMappingRuntime/Public/DMXPixelMapping.h
+++ b/Engine/Plugins/VirtualProduction/DMX/DMXPixelMapping/Source/DMXPixelMappingRuntime/Public/DMXPixelMapping.h
@@ -2,7 +2,9 @@
 
 #pragma once
 
+#include "Tickable.h"
 #include "UObject/Object.h"
+
 #include "DMXPixelMapping.generated.h"
 
 class SWidget;
@@ -10,7 +12,7 @@ class UDMXEntityFixturePatch;
 class UDMXPixelMappingBaseComponent;
 class UDMXPixelMappingOutputComponent;
 class UTexture;
-
+enum class EDMXPixelMappingResetDMXMode : uint8;
 
 /** 
  * DMX Pixel Mapping.
@@ -20,6 +22,7 @@ class UTexture;
 UCLASS(BlueprintType, Blueprintable)
 class DMXPIXELMAPPINGRUNTIME_API UDMXPixelMapping
 	: public UObject
+	, public FTickableGameObject
 {
 	using TComponentPredicate = TFunctionRef<void(UDMXPixelMappingBaseComponent*)>;
 
@@ -31,7 +34,29 @@ class DMXPIXELMAPPINGRUNTIME_API UDMXPixelMapping
 public:
 	UDMXPixelMapping();
 
+	/** Starts to send DMX */
+	void StartSendingDMX();
+
+	/** Stops sending DMX  */
+	void StopSendingDMX();
+
+	/** Pauses sending DMX */
+	void PauseSendingDMX();
+
+	/** Returns true if the pixel mapping is playing DMX. Note this returns false while it's paused */
+	bool IsPlayingDMX() const { return bIsPlayingDMX; }
+
+	/** Returns true if the pixel mapping is paused from playing DMX */
+	bool IsPaused() const { return bIsPaused; }
+
+	/** Sets how pixel mapping resets the channels it sends to when Stop Sending DMX is called. */
+	void SetResetDMXMode(EDMXPixelMappingResetDMXMode InResetDMXMode) { ResetDMXMode = InResetDMXMode; }
+
+	/** Returns how the DMX channels used by pixel mapping are reset when playing the pixel mapping is stopped */
+	EDMXPixelMappingResetDMXMode GetResetDMXMode() const { return ResetDMXMode; }
+
 	//~ Begin UObject implementation
+	virtual void Serialize(FArchive& Ar) override;
 	virtual void PostLoad() override;
 	//~ End UObject implementation
 
@@ -140,6 +165,17 @@ public:
 	 */
 	void RemoveComponent(UDMXPixelMappingBaseComponent* InComponent);
 
+	//~ Begin FTickableGameObject interface
+	virtual bool IsTickableInEditor() const { return true; }
+	virtual ETickableTickType GetTickableTickType() const override;
+	virtual bool IsTickable() const override;
+	virtual void Tick(float DeltaTime) override;
+	TStatId GetStatId() const
+	{
+		RETURN_QUICK_DECLARE_CYCLE_STAT(UDMXPixelMapping, STATGROUP_Tickables);
+	}
+	//~ End FTickableGameObject interface
+
 public:
 	/** Holds the reference to root component */
 	UPROPERTY()
@@ -185,4 +221,16 @@ public:
 	UPROPERTY()
 	TObjectPtr<UTexture> ThumbnailImage;
 #endif // WITH_EDITORONLY_DATA
+
+private:
+	/** True while the pixel mapping is sending DMX on tick */
+	bool bIsPlayingDMX = false;
+
+	/** True while the pixel mapping is paused */
+	bool bIsPaused = false;
+
+	/** Defines how DMX is reset when playback is stopped */
+	UPROPERTY()
+	EDMXPixelMappingResetDMXMode ResetDMXMode;
+
 };
