DMX: Add an IsPlayingDMX BP node to pixel mapping actor, rename new UDMXPixelMapping::bIsPlayingDMX and UDMXPixelMapping::IsPlayingDMX to IsSendingDMX for consistency with DMX Pixel Mapping Actor BP nodes

#rnx

diff --git a/Engine/Plugins/VirtualProduction/DMX/DMXPixelMapping/Source/DMXPixelMappingEditor/Private/Toolkits/DMXPixelMappingToolkit.cpp b/Engine/Plugins/VirtualProduction/DMX/DMXPixelMapping/Source/DMXPixelMappingEditor/Private/Toolkits/DMXPixelMappingToolkit.cpp
--- a/Engine/Plugins/VirtualProduction/DMX/DMXPixelMapping/Source/DMXPixelMappingEditor/Private/Toolkits/DMXPixelMappingToolkit.cpp
+++ b/Engine/Plugins/VirtualProduction/DMX/DMXPixelMapping/Source/DMXPixelMappingEditor/Private/Toolkits/DMXPixelMappingToolkit.cpp
@@ -458,7 +458,7 @@ void FDMXPixelMappingToolkit::StopPlayingDMX()
 void FDMXPixelMappingToolkit::TogglePlayPauseDMX()
 {
 	UDMXPixelMapping* PixelMapping = GetDMXPixelMapping();
-	if (PixelMapping && PixelMapping->IsPlayingDMX())
+	if (PixelMapping && PixelMapping->IsSendingDMX())
 	{
 		PixelMapping->PauseSendingDMX();
 	}
@@ -471,7 +471,7 @@ void FDMXPixelMappingToolkit::TogglePlayPauseDMX()
 void FDMXPixelMappingToolkit::TogglePlayStopDMX()
 {
 	UDMXPixelMapping* PixelMapping = GetDMXPixelMapping();
-	if (PixelMapping && PixelMapping->IsPlayingDMX())
+	if (PixelMapping && PixelMapping->IsSendingDMX())
 	{
 		PixelMapping->StopSendingDMX();
 	}
@@ -1024,7 +1024,7 @@ bool FDMXPixelMappingToolkit::IsPlayingDMX() const
 {
 	const UDMXPixelMapping* PixelMapping = GetDMXPixelMapping();
 
-	return PixelMapping && PixelMapping->IsPlayingDMX();
+	return PixelMapping && PixelMapping->IsSendingDMX();
 }
 
 #define UE_DMX_MAP_EDITOR_SETTING_TO_TOGGLE_COMMAND(EditorSettings, StructName, MemberName, Action) \
@@ -1064,19 +1064,19 @@ void FDMXPixelMappingToolkit::SetupCommands()
 		FCanExecuteAction::CreateLambda([this]
 			{
 				const UDMXPixelMapping* PixelMapping = GetDMXPixelMapping();
-				const bool bIsPlayingDMX = PixelMapping && PixelMapping->IsPlayingDMX();
+				const bool bIsSendingDMX = PixelMapping && PixelMapping->IsSendingDMX();
 				const bool bIsPaused = PixelMapping && PixelMapping->IsPaused();
 
-				return !bIsPlayingDMX && !bIsPaused;
+				return !bIsSendingDMX && !bIsPaused;
 			}),
 		FIsActionChecked(),
 		FIsActionButtonVisible::CreateLambda([this]
 			{				
 				const UDMXPixelMapping* PixelMapping = GetDMXPixelMapping();
-				const bool bIsPlayingDMX = PixelMapping && PixelMapping->IsPlayingDMX();
+				const bool bIsSendingDMX = PixelMapping && PixelMapping->IsSendingDMX();
 				const bool bIsPaused = PixelMapping && PixelMapping->IsPaused();
 
-				return !bIsPlayingDMX && !bIsPaused;
+				return !bIsSendingDMX && !bIsPaused;
 			})
 	);
 
@@ -1086,17 +1086,17 @@ void FDMXPixelMappingToolkit::SetupCommands()
 		FCanExecuteAction::CreateLambda([this] 
 			{ 
 				const UDMXPixelMapping* PixelMapping = GetDMXPixelMapping();
-				const bool bIsPlayingDMX = PixelMapping && PixelMapping->IsPlayingDMX();
+				const bool bIsSendingDMX = PixelMapping && PixelMapping->IsSendingDMX();
 
-				return bIsPlayingDMX; 
+				return bIsSendingDMX; 
 			}),
 		FIsActionChecked(),
 		FIsActionButtonVisible::CreateLambda([this]
 			{
 				const UDMXPixelMapping* PixelMapping = GetDMXPixelMapping();
-				const bool bIsPlayingDMX = PixelMapping && PixelMapping->IsPlayingDMX();
+				const bool bIsSendingDMX = PixelMapping && PixelMapping->IsSendingDMX();
 
-				return bIsPlayingDMX;
+				return bIsSendingDMX;
 			})
 	);
 
@@ -1106,19 +1106,19 @@ void FDMXPixelMappingToolkit::SetupCommands()
 		FCanExecuteAction::CreateLambda([this]
 			{
 				const UDMXPixelMapping* PixelMapping = GetDMXPixelMapping();
-				const bool bIsPlayingDMX = PixelMapping && PixelMapping->IsPlayingDMX();
+				const bool bIsSendingDMX = PixelMapping && PixelMapping->IsSendingDMX();
 				const bool bIsPaused = PixelMapping && PixelMapping->IsPaused();
 
-				return !bIsPlayingDMX && bIsPaused;
+				return !bIsSendingDMX && bIsPaused;
 			}),
 		FIsActionChecked(),
 		FIsActionButtonVisible::CreateLambda([this]
 			{
 				const UDMXPixelMapping* PixelMapping = GetDMXPixelMapping();
-				const bool bIsPlayingDMX = PixelMapping && PixelMapping->IsPlayingDMX();
+				const bool bIsSendingDMX = PixelMapping && PixelMapping->IsSendingDMX();
 				const bool bIsPaused = PixelMapping && PixelMapping->IsPaused();
 
-				return !bIsPlayingDMX && bIsPaused;
+				return !bIsSendingDMX && bIsPaused;
 			})
 	);
 
@@ -1128,10 +1128,10 @@ void FDMXPixelMappingToolkit::SetupCommands()
 		FCanExecuteAction::CreateLambda([this]
 			{
 				const UDMXPixelMapping* PixelMapping = GetDMXPixelMapping();
-				const bool bIsPlayingDMX = PixelMapping && PixelMapping->IsPlayingDMX();
+				const bool bIsSendingDMX = PixelMapping && PixelMapping->IsSendingDMX();
 				const bool bIsPaused = PixelMapping && PixelMapping->IsPaused();
 
-				return bIsPlayingDMX || bIsPaused;
+				return bIsSendingDMX || bIsPaused;
 			})
 	);
 
diff --git a/Engine/Plugins/VirtualProduction/DMX/DMXPixelMapping/Source/DMXPixelMappingRuntime/Internal/DMXPixelMappingActor.h b/Engine/Plugins/VirtualProduction/DMX/DMXPixelMapping/Source/DMXPixelMappingRuntime/Internal/DMXPixelMappingActor.h
--- a/Engine/Plugins/VirtualProduction/DMX/DMXPixelMapping/Source/DMXPixelMappingRuntime/Internal/DMXPixelMappingActor.h
+++ b/Engine/Plugins/VirtualProduction/DMX/DMXPixelMapping/Source/DMXPixelMappingRuntime/Internal/DMXPixelMappingActor.h
@@ -36,6 +36,10 @@ public:
 	UFUNCTION(BlueprintCallable, Category = "DMX Pixel Mapping")
 	void PauseSendingDMX();
 
+	/** Returns true if the asset is playing DMX */
+	UFUNCTION(BlueprintPure, Category = "DMX Pixel Mapping")
+	bool IsSendingDMX() const;
+
 	/** Sets how the pixel mapping asset resets the channels it sends to when Stop Sending DMX is called */
 	UFUNCTION(BlueprintCallable, Category = "DMX Pixel Mapping")
 	void SetStopMode(EDMXPixelMappingResetDMXMode ResetMode);
diff --git a/Engine/Plugins/VirtualProduction/DMX/DMXPixelMapping/Source/DMXPixelMappingRuntime/Private/DMXPixelMapping.cpp b/Engine/Plugins/VirtualProduction/DMX/DMXPixelMapping/Source/DMXPixelMappingRuntime/Private/DMXPixelMapping.cpp
--- a/Engine/Plugins/VirtualProduction/DMX/DMXPixelMapping/Source/DMXPixelMappingRuntime/Private/DMXPixelMapping.cpp
+++ b/Engine/Plugins/VirtualProduction/DMX/DMXPixelMapping/Source/DMXPixelMappingRuntime/Private/DMXPixelMapping.cpp
@@ -28,7 +28,7 @@ UDMXPixelMapping::UDMXPixelMapping()
 void UDMXPixelMapping::StartSendingDMX()
 {
 	bIsPaused = false;
-	bIsPlayingDMX = true;
+	bIsSendingDMX = true;
 }
 
 void UDMXPixelMapping::StopSendingDMX()
@@ -39,13 +39,13 @@ void UDMXPixelMapping::StopSendingDMX()
 	}
 
 	bIsPaused = false;
-	bIsPlayingDMX = false;
+	bIsSendingDMX = false;
 }
 
 void UDMXPixelMapping::PauseSendingDMX()
 {
 	bIsPaused = true;
-	bIsPlayingDMX = false;
+	bIsSendingDMX = false;
 }
 
 void UDMXPixelMapping::Serialize(FArchive& Ar)
@@ -191,7 +191,7 @@ ETickableTickType UDMXPixelMapping::GetTickableTickType() const
 
 bool UDMXPixelMapping::IsTickable() const
 {
-	return bIsPlayingDMX && !bIsPaused;
+	return bIsSendingDMX && !bIsPaused;
 }
 
 void UDMXPixelMapping::Tick(float DeltaTime)
diff --git a/Engine/Plugins/VirtualProduction/DMX/DMXPixelMapping/Source/DMXPixelMappingRuntime/Private/DMXPixelMappingActor.cpp b/Engine/Plugins/VirtualProduction/DMX/DMXPixelMapping/Source/DMXPixelMappingRuntime/Private/DMXPixelMappingActor.cpp
--- a/Engine/Plugins/VirtualProduction/DMX/DMXPixelMapping/Source/DMXPixelMappingRuntime/Private/DMXPixelMappingActor.cpp
+++ b/Engine/Plugins/VirtualProduction/DMX/DMXPixelMapping/Source/DMXPixelMappingRuntime/Private/DMXPixelMappingActor.cpp
@@ -46,6 +46,11 @@ void ADMXPixelMappingActor::PauseSendingDMX()
 	}
 }
 
+bool ADMXPixelMappingActor::IsSendingDMX() const
+{
+	return PixelMapping ? PixelMapping->IsSendingDMX() : false;
+}
+
 void ADMXPixelMappingActor::SetStopMode(EDMXPixelMappingResetDMXMode ResetMode)
 {
 	if (PixelMapping)
diff --git a/Engine/Plugins/VirtualProduction/DMX/DMXPixelMapping/Source/DMXPixelMappingRuntime/Public/Components/DMXPixelMappingBaseComponent.h b/Engine/Plugins/VirtualProduction/DMX/DMXPixelMapping/Source/DMXPixelMappingRuntime/Public/Components/DMXPixelMappingBaseComponent.h
--- a/Engine/Plugins/VirtualProduction/DMX/DMXPixelMapping/Source/DMXPixelMappingRuntime/Public/Components/DMXPixelMappingBaseComponent.h
+++ b/Engine/Plugins/VirtualProduction/DMX/DMXPixelMapping/Source/DMXPixelMappingRuntime/Public/Components/DMXPixelMappingBaseComponent.h
@@ -186,7 +186,7 @@ public:
 	UFUNCTION(BlueprintCallable, Category = "DMX|PixelMapping")
 	virtual void SendDMX() {};
 
-	/** Render downsample texture for this component and all children */
+	/** Renders this component and all children */
 	UFUNCTION(BlueprintCallable, Category = "DMX|PixelMapping")
 	virtual void Render() {};
 
diff --git a/Engine/Plugins/VirtualProduction/DMX/DMXPixelMapping/Source/DMXPixelMappingRuntime/Public/DMXPixelMapping.h b/Engine/Plugins/VirtualProduction/DMX/DMXPixelMapping/Source/DMXPixelMappingRuntime/Public/DMXPixelMapping.h
--- a/Engine/Plugins/VirtualProduction/DMX/DMXPixelMapping/Source/DMXPixelMappingRuntime/Public/DMXPixelMapping.h
+++ b/Engine/Plugins/VirtualProduction/DMX/DMXPixelMapping/Source/DMXPixelMappingRuntime/Public/DMXPixelMapping.h
@@ -44,7 +44,7 @@ public:
 	void PauseSendingDMX();
 
 	/** Returns true if the pixel mapping is playing DMX. Note this returns false while it's paused */
-	bool IsPlayingDMX() const { return bIsPlayingDMX; }
+	bool IsSendingDMX() const { return bIsSendingDMX; }
 
 	/** Returns true if the pixel mapping is paused from playing DMX */
 	bool IsPaused() const { return bIsPaused; }
@@ -224,7 +224,7 @@ public:
 
 private:
 	/** True while the pixel mapping is sending DMX on tick */
-	bool bIsPlayingDMX = false;
+	bool bIsSendingDMX = false;
 
 	/** True while the pixel mapping is paused */
 	bool bIsPaused = false;
