DMX: Fix DMXPixelMappingActor '#undef LOCTEXT_NAMESPACE' without a corresponding '#define LOCTEXT_NAMESPACE'

#rnx

diff --git a/Engine/Plugins/VirtualProduction/DMX/DMXPixelMapping/Source/DMXPixelMappingRuntime/Private/DMXPixelMappingActor.cpp b/Engine/Plugins/VirtualProduction/DMX/DMXPixelMapping/Source/DMXPixelMappingRuntime/Private/DMXPixelMappingActor.cpp
--- a/Engine/Plugins/VirtualProduction/DMX/DMXPixelMapping/Source/DMXPixelMappingRuntime/Private/DMXPixelMappingActor.cpp
+++ b/Engine/Plugins/VirtualProduction/DMX/DMXPixelMapping/Source/DMXPixelMappingRuntime/Private/DMXPixelMappingActor.cpp
@@ -125,5 +125,3 @@ void ADMXPixelMappingActor::ApplySendDMXInEditorState()
 	}
 }
 #endif // WITH_EDITOR
-
-#undef LOCTEXT_NAMESPACE
